Rename the existing Project Launcher to Legacy Project Launcher, and present the new tool as Project Launcher.
-Legacy Project Launcher now kept under the Tools -> Platforms tab
- new tool will be moved & renamed in a follow up change.

diff --git a/Engine/Plugins/Developer/CustomLaunchUI/CustomLaunchUI.uplugin b/Engine/Plugins/Developer/CustomLaunchUI/CustomLaunchUI.uplugin
--- a/Engine/Plugins/Developer/CustomLaunchUI/CustomLaunchUI.uplugin
+++ b/Engine/Plugins/Developer/CustomLaunchUI/CustomLaunchUI.uplugin
@@ -2,7 +2,7 @@
 	"FileVersion": 3,
 	"Version": 1,
 	"VersionName": "0.1",
-	"FriendlyName": "Custom Launch UI",
+	"FriendlyName": "Project Launcher",
 	"Description": "Configure custom project launch profiles.",
 	"CreatedBy": "Epic Games, Inc.",
 	"Category": "Platform",
diff --git a/Engine/Plugins/Developer/CustomLaunchUI/Source/CustomLaunchUI/Private/CustomLaunchUIModule.cpp b/Engine/Plugins/Developer/CustomLaunchUI/Source/CustomLaunchUI/Private/CustomLaunchUIModule.cpp
--- a/Engine/Plugins/Developer/CustomLaunchUI/Source/CustomLaunchUI/Private/CustomLaunchUIModule.cpp
+++ b/Engine/Plugins/Developer/CustomLaunchUI/Source/CustomLaunchUI/Private/CustomLaunchUIModule.cpp
@@ -16,7 +16,7 @@
 
 #define LOCTEXT_NAMESPACE "FCustomLaunchUIModule"
 
-static const FName TabName("CustomLaunchUI");
+static const FName TabName("ProjectLauncher");
 
 namespace CustomLaunchUI
 {
@@ -43,11 +43,10 @@ public:
 
 		// register main tab
 		FGlobalTabmanager::Get()->RegisterNomadTabSpawner(TabName, FOnSpawnTab::CreateRaw(this, &FCustomLaunchUIModule::HandleSpawnTab))
-			.SetDisplayName(LOCTEXT("SpawnTabLabel", "Custom Launch Profiles"))
-			.SetTooltipText(LOCTEXT("SpawnTabToolTip", "Configure custom launch profiles for advanced packaging, deploying and launching of your projects"))
+			.SetDisplayName(LOCTEXT("SpawnTabLabel", "Project Launcher"))
+			.SetTooltipText(LOCTEXT("SpawnTabToolTip", "Configure custom launch profiles for advanced packaging, deploying and launching of your project"))
 			.SetIcon(FSlateIcon(FAppStyle::GetAppStyleSetName(), "Launcher.TabIcon"))
-			.SetGroup(WorkspaceMenu::GetMenuStructure().GetDeveloperToolsPlatformsCategory())
-			.SetMenuType(ETabSpawnerMenuType::Hidden); // hide it for now - it will be invoked by TurnkeySupportModule
+			.SetGroup(WorkspaceMenu::GetMenuStructure().GetToolsCategory());
 	}
 
 	virtual void ShutdownModule() override
diff --git a/Engine/Plugins/Developer/CustomLaunchUI/Source/CustomLaunchUI/Private/Widgets/Profiles/SCustomLaunchProfilesPanel.cpp b/Engine/Plugins/Developer/CustomLaunchUI/Source/CustomLaunchUI/Private/Widgets/Profiles/SCustomLaunchProfilesPanel.cpp
--- a/Engine/Plugins/Developer/CustomLaunchUI/Source/CustomLaunchUI/Private/Widgets/Profiles/SCustomLaunchProfilesPanel.cpp
+++ b/Engine/Plugins/Developer/CustomLaunchUI/Source/CustomLaunchUI/Private/Widgets/Profiles/SCustomLaunchProfilesPanel.cpp
@@ -495,7 +495,8 @@ TSharedRef<SWidget> SCustomLaunchProfilesPanel::CreateAdvancedProfileWarningWidg
 			[
 				SNew(STextBlock)
 				.ColorAndOpacity(FLinearColor::White)
-				.Text(LOCTEXT("AdvancedProfileWarningDetail", "These settings may be multiple cook platforms or cultures, or distribution packages, DLC or patch builds.\nPlease use Project Launcher to edit the advanced properties in this profile."))
+				.Text(LOCTEXT("AdvancedProfileWarningDetail", "These settings may be multiple cook platforms or cultures, or distribution packages, DLC or patch builds. Please use Legacy Project Launcher to edit the advanced properties in this profile."))
+				.AutoWrapText(true)
 			]
 		]
 
@@ -505,7 +506,7 @@ TSharedRef<SWidget> SCustomLaunchProfilesPanel::CreateAdvancedProfileWarningWidg
 		.VAlign(VAlign_Center)
 		[
 			SNew(SButton)
-			.OnClicked_Lambda( []() { FGlobalTabmanager::Get()->TryInvokeTab(FTabId("ProjectLauncher")); return FReply::Handled(); })
+			.OnClicked_Lambda( []() { FGlobalTabmanager::Get()->TryInvokeTab(FTabId("LegacyProjectLauncher")); return FReply::Handled(); })
 			.ButtonStyle(FCustomLaunchUIStyle::Get(), "HoverHintOnly")
 			.ToolTipText(LOCTEXT("OpenDeviceProjectLauncherTip", "Open the legacy Project Launcher."))
 			[
diff --git a/Engine/Source/Developer/LegacyProjectLauncher/Private/Models/ProjectLauncherModel.h b/Engine/Source/Developer/LegacyProjectLauncher/Private/Models/ProjectLauncherModel.h
new file mode 100644
--- /dev/null
+++ b/Engine/Source/Developer/LegacyProjectLauncher/Private/Models/ProjectLauncherModel.h
@@ -0,0 +1,264 @@
+// Copyright Epic Games, Inc. All Rights Reserved.
+
+#pragma once
+
+#include "CoreTypes.h"
+#include "Containers/Array.h"
+#include "Delegates/Delegate.h"
+#include "GameProjectHelper.h"
+#include "ILauncher.h"
+#include "ILauncherProfileManager.h"
+#include "Misc/Paths.h"
+#include "Misc/ConfigCacheIni.h"
+#include "Misc/App.h"
+#include "Templates/SharedPointer.h"
+
+
+class FProjectLauncherModel;
+class SProjectLauncher;
+
+
+namespace ELauncherPanels
+{
+	/**
+	 * Enumerates available launcher panels.
+	 */
+	enum Type
+	{
+		NoTask = 0,
+		Launch,
+		ProfileEditor,
+		Progress,
+		End,
+	};
+}
+
+
+/**
+ * Delegate type for launcher profile selection changes.
+ *
+ * The first parameter is the newly selected profile (or nullptr if none was selected).
+ * The second parameter is the old selected profile (or nullptr if none was previous selected).
+ */
+DECLARE_MULTICAST_DELEGATE_TwoParams(FOnSelectedSProjectLauncherProfileChanged, const ILauncherProfilePtr&, const ILauncherProfilePtr&)
+
+
+/**
+ * Implements the data model for the session launcher.
+ */
+class FProjectLauncherModel
+{
+public:
+
+	/**
+	 * Creates and initializes a new instance.
+	 *
+	 * @param InDeviceProxyManager The device proxy manager to use.
+	 * @param InLauncher The launcher to use.
+	 * @param InProfileManager The profile manager to use.
+	 */
+	FProjectLauncherModel(const TSharedRef<ITargetDeviceProxyManager>& InDeviceProxyManager, const TSharedRef<ILauncher>& InLauncher, const TSharedRef<ILauncherProfileManager>& InProfileManager)
+		: DeviceProxyManager(InDeviceProxyManager)
+		, SProjectLauncher(InLauncher)
+		, ProfileManager(InProfileManager)
+	{
+		ProfileManager->OnProfileAdded().AddRaw(this, &FProjectLauncherModel::HandleProfileManagerProfileAdded);
+		ProfileManager->OnProfileRemoved().AddRaw(this, &FProjectLauncherModel::HandleProfileManagerProfileRemoved);
+
+		LoadConfig();
+	}
+
+	/** Destructor. */
+	~FProjectLauncherModel()
+	{
+		ProfileManager->OnProfileAdded().RemoveAll(this);
+		ProfileManager->OnProfileRemoved().RemoveAll(this);
+
+		SaveConfig();
+	}
+
+public:
+
+	/**
+	 * Gets the model's device proxy manager.
+	 *
+	 * @return Device proxy manager.
+	 */
+	const TSharedRef<ITargetDeviceProxyManager>& GetDeviceProxyManager() const
+	{
+		return DeviceProxyManager;
+	}
+
+	/**
+	 * Gets the model's launcher.
+	 *
+	 * @return SProjectLauncher.
+	 */
+	const TSharedRef<ILauncher>& GetSProjectLauncher() const
+	{
+		return SProjectLauncher;
+	}
+
+	/**
+	 * Get the model's profile manager.
+	 *
+	 * @return Profile manager.
+	 * @see GetProfiles
+	 */
+	const TSharedRef<ILauncherProfileManager>& GetProfileManager() const
+	{
+		return ProfileManager;
+	}
+
+	/**
+	 * Gets the active profile.
+	 *
+	 * @return The active profile.
+	 * @see GetProfiles, GetProfileManager, SelectProfile
+	 */
+	const TSharedPtr<ILauncherProfile>& GetSelectedProfile() const
+	{
+		return SelectedProfile;
+	}
+
+	/**
+	 * Sets the active profile.
+	 *
+	 * @param InProfile The profile to assign as active,, or nullptr to deselect all.
+	 * @see GetSelectedProfile
+	 */
+	void SelectProfile(const TSharedPtr<ILauncherProfile>& Profile)
+	{
+		if (!Profile.IsValid() || ProfileManager->GetAllProfiles().Contains(Profile))
+		{
+			if (Profile != SelectedProfile)
+			{
+				TSharedPtr<ILauncherProfile> PreviousProfile = SelectedProfile;
+				SelectedProfile = Profile;
+
+				ProfileSelectedDelegate.Broadcast(Profile, PreviousProfile);
+			}
+		}
+	}
+
+public:
+
+	/**
+	 * Returns a delegate to be invoked when the profile list has been modified.
+	 *
+	 * @return The delegate.
+	 * @see OnProfileSelected
+	 */
+	FSimpleMulticastDelegate& OnProfileListChanged()
+	{
+		return ProfileListChangedDelegate;
+	}
+
+	/**
+	 * Returns a delegate to be invoked when the selected profile changed.
+	 *
+	 * @return The delegate.
+	 * @see OnProfileListChanged
+	 */
+	FOnSelectedSProjectLauncherProfileChanged& OnProfileSelected()
+	{
+		return ProfileSelectedDelegate;
+	}
+
+protected:
+
+	/*
+	 * Load all profiles from disk.
+	 *
+	 * @see SaveConfig
+	 */
+	void LoadConfig()
+	{
+		// restore the previous project selection
+		FString ProjectPath;
+
+		if (FPaths::IsProjectFilePathSet())
+		{
+			ProjectPath = FPaths::GetProjectFilePath();
+		}
+		else if (FGameProjectHelper::IsGameAvailable(FApp::GetProjectName()))
+		{
+			ProjectPath = FPaths::RootDir() / FApp::GetProjectName() / FApp::GetProjectName() + TEXT(".uproject");
+		}
+		else if (GConfig != nullptr)
+		{
+			GConfig->GetString(TEXT("FProjectLauncherModel"), TEXT("SelectedProjectPath"), ProjectPath, GEngineIni);
+		}
+
+		ProfileManager->SetProjectPath(ProjectPath);
+	}
+
+	/*
+	 * Saves all profiles to disk.
+
+	 * @see LoadConfig
+	 */
+	void SaveConfig()
+	{
+		if (GConfig != nullptr && !FPaths::IsProjectFilePathSet() && !FGameProjectHelper::IsGameAvailable(FApp::GetProjectName()))
+		{
+			FString ProjectPath = ProfileManager->GetProjectPath();
+			GConfig->SetString(TEXT("FProjectLauncherModel"), TEXT("SelectedProjectPath"), *ProjectPath, GEngineIni);
+		}
+	}
+
+private:
+
+	/** Callback for when a profile was added to the profile manager. */
+	void HandleProfileManagerProfileAdded(const TSharedRef<ILauncherProfile>& Profile)
+	{
+		ProfileListChangedDelegate.Broadcast();
+
+		SelectProfile(Profile);
+	}
+
+	/** Callback for when a profile was removed from the profile manager. */
+	void HandleProfileManagerProfileRemoved(const TSharedRef<ILauncherProfile>& Profile)
+	{
+		ProfileListChangedDelegate.Broadcast();
+
+		if (Profile == SelectedProfile)
+		{
+			const TArray<TSharedPtr<ILauncherProfile>>& Profiles = ProfileManager->GetAllProfiles();
+
+			if (Profiles.Num() > 0)
+			{
+				SelectProfile(Profiles[0]);
+			}
+			else
+			{
+				SelectProfile(nullptr);
+			}
+		}
+	}
+
+private:
+
+	/** Pointer to the device proxy manager. */
+	TSharedRef<ITargetDeviceProxyManager> DeviceProxyManager;
+
+	/** Pointer to the launcher. */
+	TSharedRef<ILauncher> SProjectLauncher;
+
+	/** Pointer to the profile manager. */
+	TSharedRef<ILauncherProfileManager> ProfileManager;
+
+	/** Pointer to active profile. */
+	TSharedPtr<ILauncherProfile> SelectedProfile;
+
+private:
+
+	/** A delegate to be invoked when the project path has been modified. */
+	FSimpleMulticastDelegate ProjectPathChangedDelegate;
+
+	/** A delegate to be invoked when the profile list has been modified. */
+	FSimpleMulticastDelegate ProfileListChangedDelegate;
+
+	/** A delegate to be invoked when the selected profile changed. */
+	FOnSelectedSProjectLauncherProfileChanged ProfileSelectedDelegate;
+};
diff --git a/Engine/Source/Developer/LegacyProjectLauncher/Private/Widgets/SProjectLauncher.cpp b/Engine/Source/Developer/LegacyProjectLauncher/Private/Widgets/SProjectLauncher.cpp
new file mode 100644
--- /dev/null
+++ b/Engine/Source/Developer/LegacyProjectLauncher/Private/Widgets/SProjectLauncher.cpp
@@ -0,0 +1,437 @@
+// Copyright Epic Games, Inc. All Rights Reserved.
+
+#include "Widgets/SProjectLauncher.h"
+#include "Framework/Commands/UIAction.h"
+#include "Framework/Docking/TabManager.h"
+#include "Framework/MultiBox/MultiBoxBuilder.h"
+#include "SlateOptMacros.h"
+#include "Styling/SlateTypes.h"
+#include "Widgets/SBoxPanel.h"
+#include "Widgets/Images/SImage.h"
+#include "Widgets/Input/SButton.h"
+#include "Widgets/Input/SCheckBox.h"
+#include "Widgets/Input/SComboButton.h"
+#include "Widgets/Layout/SBorder.h"
+#include "Widgets/Layout/SSplitter.h"
+#include "Widgets/Layout/SWidgetSwitcher.h"
+#include "Widgets/Text/STextBlock.h"
+
+#include "Models/ProjectLauncherCommands.h"
+#include "Widgets/Deploy/SProjectLauncherSimpleDeviceListView.h"
+#include "Widgets/Profile/SProjectLauncherProfileListView.h"
+#include "Widgets/Progress/SProjectLauncherProgress.h"
+#include "Widgets/Project/SProjectLauncherProjectPicker.h"
+#include "Widgets/Shared/SProjectLauncherBuildTargetSelector.h"
+#include "Widgets/Settings/SProjectLauncherSettings.h"
+#include "Styling/AppStyle.h"
+#include "Styling/StyleColors.h"
+#include "SPositiveActionButton.h"
+
+#define LOCTEXT_NAMESPACE "SProjectLauncher"
+
+
+/* SProjectLauncher structors
+ *****************************************************************************/
+
+SProjectLauncher::SProjectLauncher()
+	:bAdvanced(false)
+{
+	if (GConfig != NULL)
+	{
+		GConfig->GetBool(TEXT("FProjectLauncher"), TEXT("AdvancedMode"), bAdvanced, GEngineIni);
+	}
+}
+
+
+SProjectLauncher::~SProjectLauncher()
+{
+	if (GConfig != NULL)
+	{
+		GConfig->SetBool(TEXT("FProjectLauncher"), TEXT("AdvancedMode"), bAdvanced, GEngineIni);
+	}
+
+	if (LauncherWorker.IsValid())
+	{
+		LauncherWorker->Cancel();
+		FPlatformProcess::Sleep(0.5f);
+	}
+}
+
+
+/* SProjectLauncher interface
+ *****************************************************************************/
+
+BEGIN_SLATE_FUNCTION_BUILD_OPTIMIZATION
+void SProjectLauncher::Construct(const FArguments& InArgs, const TSharedRef<SDockTab>& ConstructUnderMajorTab, const TSharedPtr<SWindow>& ConstructUnderWindow, const TSharedRef<FProjectLauncherModel>& InModel)
+{
+	FProjectLauncherCommands::Register();
+
+	Model = InModel;
+
+	// create & initialize main menu bar
+	TSharedRef<FWorkspaceItem> RootMenuGroup = FWorkspaceItem::NewGroup(LOCTEXT("RootMenuGroup", "Root"));
+
+	FMenuBarBuilder MenuBarBuilder = FMenuBarBuilder(TSharedPtr<FUICommandList>());
+	MenuBarBuilder.AddPullDownMenu(
+		LOCTEXT("WindowMenuLabel", "Window"),
+		FText::GetEmpty(),
+		FNewMenuDelegate::CreateStatic(&SProjectLauncher::FillWindowMenu, RootMenuGroup),
+		"Window"
+	);
+
+	ChildSlot
+	[
+		SAssignNew(WidgetSwitcher, SWidgetSwitcher)
+		.WidgetIndex((int32)ELauncherPanels::Launch)
+		
+		// Empty Panel
+		+ SWidgetSwitcher::Slot()
+		[
+			SNew(SBorder)
+		]
+
+		// SProjectLauncher Panel
+		+ SWidgetSwitcher::Slot()
+		[
+			SNew(SSplitter)
+			.Orientation(Orient_Vertical)
+
+			// Simple SProjectLauncher
+			+ SSplitter::Slot()
+			.Value(0.55f)
+			[
+				SNew(SVerticalBox)
+
+				+ SVerticalBox::Slot()
+				.AutoHeight()
+				[
+					SNew(SBorder)
+					.BorderImage(FAppStyle::Get().GetBrush("Brushes.Panel"))
+					[
+						SNew(SHorizontalBox)
+
+						// Project Bar
+						+ SHorizontalBox::Slot()
+						.AutoWidth()
+						.Padding(0, 4)
+						[
+							SNew(SProjectLauncherProjectPicker, InModel)
+						]
+
+						// Build Target
+						+ SHorizontalBox::Slot()
+						.AutoWidth()
+						.Padding(0,4)
+						[
+							SNew(SProjectLauncherBuildTargetSelector, InModel)
+							.UseProfile(false)
+						]
+
+						// Advanced Button
+						+ SHorizontalBox::Slot()
+						.AutoWidth()
+						.VAlign(VAlign_Center)
+						.Padding(6, 0, 0, 0)
+						[
+							
+							SNew(SCheckBox)
+							.Style(FAppStyle::Get(), "ToggleButtonCheckbox")
+							.IsFocusable(true)
+							.ToolTipText(LOCTEXT("ToggleAdvancedOptionsToolTipText", "Toggles Advanced Options"))
+							.OnCheckStateChanged(this, &SProjectLauncher::OnAdvancedChanged)
+							.IsChecked(this, &SProjectLauncher::OnIsAdvanced)
+							[
+								SNew(SHorizontalBox)
+								// Icon
+								+ SHorizontalBox::Slot()
+								.VAlign(VAlign_Center)
+								.AutoWidth()
+								[
+									SNew(SImage)
+									.Image(this, &SProjectLauncher::GetAdvancedToggleBrush)
+									.DesiredSizeOverride(FVector2D(16, 16))
+									.ColorAndOpacity(FSlateColor::UseForeground())
+								]
+
+								// Text
+								+ SHorizontalBox::Slot()
+								.AutoWidth()
+								.VAlign(VAlign_Center)
+								.Padding(4, 0, 4, 0)
+								[
+									SNew(STextBlock)
+									.Text(LOCTEXT("AdvancedButton", "Show Advanced"))
+								]
+							]
+						]
+					]
+				]
+
+				+ SVerticalBox::Slot()
+				.FillHeight(1)
+				.Padding(0)
+				[
+					SNew(SBorder)
+					.BorderImage(FAppStyle::Get().GetBrush("Brushes.Recessed"))
+					[
+						SAssignNew(LaunchList, SProjectLauncherSimpleDeviceListView, InModel)
+						.OnProfileRun(this, &SProjectLauncher::OnProfileRun)
+						.IsAdvanced(this, &SProjectLauncher::GetIsAdvanced)
+					]
+				]
+			]
+
+			+ SSplitter::Slot()
+			.Value(0.45f)
+			[
+				SNew(SBorder)
+				.Padding(0)
+				.BorderImage(FAppStyle::Get().GetBrush("Brushes.Recessed"))
+				[
+					SNew(SVerticalBox)
+
+					+ SVerticalBox::Slot()
+					.AutoHeight()
+					.Padding(0)
+					[
+						SNew(SBorder)
+						.BorderImage(FAppStyle::Get().GetBrush("Brushes.Panel"))
+						[
+							SNew(SHorizontalBox)
+
+							+ SHorizontalBox::Slot()
+							.VAlign(VAlign_Center)
+							.FillWidth(1.0f)
+							.Padding(14, 0, 0, 0)
+							[
+								SNew(STextBlock)
+								.Text(LOCTEXT("ProjectLauncherCustomProfilesTitle", "Custom Launch Profiles"))
+								.Font(FAppStyle::Get().GetFontStyle("NormalFontBold"))
+							]
+
+							
+
+							+ SHorizontalBox::Slot()
+							.HAlign(HAlign_Right)
+							.AutoWidth()
+							.Padding(0, 0, 16, 0)
+							[
+								SNew(SPositiveActionButton)
+								.Text(LOCTEXT("AddButtonLabel", "Add"))
+								.ToolTipText(LOCTEXT("AddFilterToolTip", "Add a new custom launch profile using wizard"))
+								.OnGetMenuContent(this, &SProjectLauncher::MakeProfileWizardsMenu)
+							]
+						]
+					]
+					+ SVerticalBox::Slot()
+					.FillHeight(1)
+					.Padding(2)
+					[
+						SAssignNew(ProfileList, SBorder)
+						.BorderImage(FAppStyle::Get().GetBrush("Brushes.Recessed"))
+						.Padding(0)
+						[
+							// Simple Launch List
+							SNew(SProjectLauncherProfileListView, InModel)
+							.OnProfileEdit(this, &SProjectLauncher::OnProfileEdit)
+							.OnProfileRun(this, &SProjectLauncher::OnProfileRun)
+							.OnProfileDelete(this, &SProjectLauncher::OnProfileDelete)
+						]
+					]
+				]
+			]
+		]
+		
+		// Launch Settings
+		+ SWidgetSwitcher::Slot()
+		[
+			SAssignNew(ProfileSettingsPanel, SProjectLauncherSettings, InModel)
+			.OnCloseClicked(this, &SProjectLauncher::OnProfileSettingsClose)
+			.OnDeleteClicked(this, &SProjectLauncher::OnProfileDelete)
+		]
+
+		// Progress Panel
+		+ SWidgetSwitcher::Slot()
+		[
+			SAssignNew(ProgressPanel, SProjectLauncherProgress)
+			.OnCloseClicked(this, &SProjectLauncher::OnProgressClose)
+			.OnRerunClicked(this, &SProjectLauncher::OnRerunClicked)
+		]
+	];
+}
+END_SLATE_FUNCTION_BUILD_OPTIMIZATION
+
+
+/* SProjectLauncher implementation
+*****************************************************************************/
+
+void SProjectLauncher::FillWindowMenu(FMenuBuilder& MenuBuilder, TSharedRef<FWorkspaceItem> RootMenuGroup)
+{
+#if !WITH_EDITOR
+	MenuBuilder.BeginSection("WindowGlobalTabSpawners", LOCTEXT("UfeMenuGroup", "Unreal Frontend"));
+	{
+		FGlobalTabmanager::Get()->PopulateTabSpawnerMenu(MenuBuilder, RootMenuGroup);
+	}
+	MenuBuilder.EndSection();
+#endif //!WITH_EDITOR
+}
+
+
+/* SProjectLauncher callbacks
+*****************************************************************************/
+
+void SProjectLauncher::OnAdvancedChanged(const ECheckBoxState NewCheckedState)
+{
+	bAdvanced = (NewCheckedState == ECheckBoxState::Checked);
+}
+
+
+ECheckBoxState SProjectLauncher::OnIsAdvanced() const
+{
+	return (bAdvanced) ? ECheckBoxState::Checked : ECheckBoxState::Unchecked;
+}
+
+
+const FSlateBrush* SProjectLauncher::GetAdvancedToggleBrush() const
+{
+	return FAppStyle::Get().GetBrush("Icons.Advanced");
+}
+
+
+bool SProjectLauncher::GetIsAdvanced() const
+{
+	return bAdvanced;
+}
+
+
+void SProjectLauncher::OnProfileEdit(const ILauncherProfileRef& Profile)
+{
+	Model->SelectProfile(Profile);
+	WidgetSwitcher->SetActiveWidgetIndex(ELauncherPanels::ProfileEditor);
+}
+
+
+void SProjectLauncher::OnProfileRun(const ILauncherProfileRef& Profile)
+{
+	LauncherProfile = Profile;
+	LauncherWorker = Model->GetSProjectLauncher()->Launch(Model->GetDeviceProxyManager(), Profile);
+	
+	if (LauncherWorker.IsValid())
+	{
+		ProgressPanel->SetLauncherWorker(LauncherWorker.ToSharedRef());
+		WidgetSwitcher->SetActiveWidgetIndex(ELauncherPanels::Progress);
+	}
+}
+
+
+void SProjectLauncher::OnProfileDelete(const ILauncherProfileRef& Profile)
+{
+	Model->GetProfileManager()->RemoveProfile(Profile);
+}
+
+
+void SProjectLauncher::OnAddCustomLaunchProfileClicked()
+{
+	ILauncherProfileRef Profile = Model->GetProfileManager()->AddNewProfile();
+	
+	OnProfileEdit(Profile);
+
+	ProfileSettingsPanel->EnterEditMode();
+}
+
+
+TSharedRef<SWidget> SProjectLauncher::MakeProfileWizardsMenu()
+{
+	FMenuBuilder MenuBuilder(true, NULL);
+
+	MenuBuilder.BeginSection("Create", LOCTEXT("CreateSection", "CREATE"));
+
+	MenuBuilder.AddMenuEntry(
+		LOCTEXT("CustomProfileLabel", "Create Custom Profile"),
+		LOCTEXT("CustomProfileDescription", "Add a new custom launch profile."),
+		FSlateIcon(),
+		FUIAction(
+			FExecuteAction::CreateSP(this, &SProjectLauncher::OnAddCustomLaunchProfileClicked),
+			FCanExecuteAction()
+		)
+	);
+
+	MenuBuilder.EndSection();
+
+	MenuBuilder.BeginSection("CreateFromPreset", LOCTEXT("CreateFromPreset", "CREATE FROM PRESET"));
+
+	const TArray<ILauncherProfileWizardPtr>& Wizards = Model->GetProfileManager()->GetProfileWizards();
+	for (const ILauncherProfileWizardPtr& Wizard : Wizards)
+	{
+		FText WizardName = Wizard->GetName();
+		FText WizardDescription = Wizard->GetDescription();
+				
+		MenuBuilder.AddMenuEntry(
+			WizardName,
+			WizardDescription,
+			FSlateIcon(),
+			FUIAction(
+				FExecuteAction::CreateSP(this, &SProjectLauncher::ExecProfileWizard, Wizard),
+				FCanExecuteAction()
+				)
+			);
+	}
+
+	MenuBuilder.EndSection();
+
+	return MenuBuilder.MakeWidget();
+}
+
+
+void SProjectLauncher::ExecProfileWizard(ILauncherProfileWizardPtr InWizard)
+{
+	InWizard->HandleCreateLauncherProfile(Model->GetProfileManager());
+}
+
+
+FReply SProjectLauncher::OnProfileSettingsClose()
+{
+	if (LauncherWorker.IsValid())
+	{
+		LauncherWorker->Cancel();
+	}
+	LauncherProfile.Reset();
+
+	WidgetSwitcher->SetActiveWidgetIndex(ELauncherPanels::Launch);
+
+	return FReply::Handled();
+}
+
+
+FReply SProjectLauncher::OnProgressClose()
+{
+	if (LauncherWorker.IsValid())
+	{
+		LauncherWorker->Cancel();
+	}
+	LauncherProfile.Reset();
+
+	WidgetSwitcher->SetActiveWidgetIndex(ELauncherPanels::Launch);
+
+	return FReply::Handled();
+}
+
+
+FReply SProjectLauncher::OnRerunClicked()
+{
+	if (LauncherWorker.IsValid())
+	{
+		LauncherWorker->Cancel();
+	}
+	LauncherWorker = Model->GetSProjectLauncher()->Launch(Model->GetDeviceProxyManager(), LauncherProfile.ToSharedRef());
+
+	if (LauncherWorker.IsValid())
+	{
+		ProgressPanel->SetLauncherWorker(LauncherWorker.ToSharedRef());
+	}
+
+	return FReply::Handled();
+}
+
+
+#undef LOCTEXT_NAMESPACE
diff --git a/Engine/Source/Developer/LegacyProjectLauncher/Private/Widgets/SProjectLauncher.h b/Engine/Source/Developer/LegacyProjectLauncher/Private/Widgets/SProjectLauncher.h
new file mode 100644
--- /dev/null
+++ b/Engine/Source/Developer/LegacyProjectLauncher/Private/Widgets/SProjectLauncher.h
@@ -0,0 +1,144 @@
+// Copyright Epic Games, Inc. All Rights Reserved.
+
+#pragma once
+
+#include "CoreTypes.h"
+#include "Framework/Docking/WorkspaceItem.h"
+#include "Input/Reply.h"
+#include "Layout/Visibility.h"
+#include "SlateFwd.h"
+#include "Templates/SharedPointer.h"
+#include "Widgets/DeclarativeSyntaxSupport.h"
+#include "Widgets/SCompoundWidget.h"
+
+#include "Models/ProjectLauncherModel.h"
+
+class FMenuBuilder;
+class SBorder;
+class SProjectLauncherProgress;
+class SProjectLauncherSettings;
+class SProjectLauncherSimpleDeviceListView;
+class SWindow;
+
+struct FSlateBrush;
+
+enum class ECheckBoxState : uint8;
+
+
+/**
+ * Implements a Slate widget for the launcher user interface.
+ */
+class SProjectLauncher
+	: public SCompoundWidget
+{
+public:
+
+	SLATE_BEGIN_ARGS(SProjectLauncher) { }
+
+		/** Exposes a delegate to be invoked when the launcher has closed. */
+		SLATE_EVENT(FSimpleDelegate, OnClosed)
+
+	SLATE_END_ARGS()
+
+public:
+
+	/** Constructor. */
+	SProjectLauncher();
+
+	/** Destructor. */
+	~SProjectLauncher();
+
+public:
+
+	/**
+	 * Constructs the widget.
+	 *
+	 * @param InArgs The Slate argument list.
+	 * @param ConstructUnderMajorTab The major tab which will contain the session front-end.
+	 * @param ConstructUnderWindow The window in which this widget is being constructed.
+	 * @param InModel The view model to use.
+	 */
+	void Construct(const FArguments& InArgs, const TSharedRef<SDockTab>& ConstructUnderMajorTab, const TSharedPtr<SWindow>& ConstructUnderWindow, const TSharedRef<FProjectLauncherModel>& InModel);
+
+protected:
+
+	/**
+	 * Fills the Window menu with menu items.
+	 *
+	 * @param MenuBuilder The multi-box builder that should be filled with content for this pull-down menu.
+	 * @param RootMenuGroup The root menu group.
+	 */
+	static void FillWindowMenu(FMenuBuilder& MenuBuilder, TSharedRef<FWorkspaceItem> RootMenuGroup);
+
+private:
+
+	/** Callback for toggling the advanced mode. */
+	void OnAdvancedChanged(const ECheckBoxState NewCheckedState);
+
+	/** Callback to determine if we are in advanced mode. */
+	ECheckBoxState OnIsAdvanced() const;
+
+	/** Get advanced toggle button brush. */
+	const FSlateBrush* GetAdvancedToggleBrush() const;
+
+	/** Callback for whether advanced options should be shown. */
+	bool GetIsAdvanced() const;
+
+	/** Callback for editing a profile. */
+	void OnProfileEdit(const ILauncherProfileRef& Profile);
+
+	/** Callback for running a profile. */
+	void OnProfileRun(const ILauncherProfileRef& Profile);
+
+	/** Callback for deleting a profile. */
+	void OnProfileDelete(const ILauncherProfileRef& Profile);
+
+	/** Callback for clicking the Add New Launch Profile. */
+	void OnAddCustomLaunchProfileClicked();
+
+	/** Callback for filling profile wizard menu. */
+	TSharedRef<SWidget> MakeProfileWizardsMenu();
+
+	/** Execute specified profile wizard. */
+	void ExecProfileWizard(ILauncherProfileWizardPtr InWizard);
+
+	/** Callback for when the settings panel is closed. */
+	FReply OnProfileSettingsClose();
+
+	/** Callback for when the progress panel is closed. */
+	FReply OnProgressClose();
+
+	/** Callback for when the progress panel requests the profile to be re run. */
+	FReply OnRerunClicked();
+
+private:
+
+	/** The current launcher worker, if any. */
+	ILauncherWorkerPtr LauncherWorker;
+
+	/** The Launcher profile the launcherWorker is running. */
+	ILauncherProfilePtr LauncherProfile;
+
+	/** Holds a pointer to the view model. */
+	TSharedPtr<FProjectLauncherModel> Model;
+
+	/** The profile settings panel. */
+	TSharedPtr<SProjectLauncherSettings> ProfileSettingsPanel;
+
+	/** The progress panel. */
+	TSharedPtr<SProjectLauncherProgress> ProgressPanel;
+
+	/** The widget switcher. */
+	TSharedPtr<SWidgetSwitcher> WidgetSwitcher;
+
+	/** Contains the launch list widgets. */
+//	TSharedPtr<SBorder> LaunchList;
+
+	TSharedPtr<SProjectLauncherSimpleDeviceListView> LaunchList;
+
+	/** Contains the profile list widgets. */
+	TSharedPtr<SBorder> ProfileList;
+
+	/** Whether we are showing advanced options. */
+	bool bAdvanced;
+};
diff --git a/Engine/Source/Developer/LegacyProjectLauncher/Private/Widgets/Shared/ProjectLauncherDelegates.h b/Engine/Source/Developer/LegacyProjectLauncher/Private/Widgets/Shared/ProjectLauncherDelegates.h
new file mode 100644
--- /dev/null
+++ b/Engine/Source/Developer/LegacyProjectLauncher/Private/Widgets/Shared/ProjectLauncherDelegates.h
@@ -0,0 +1,14 @@
+// Copyright Epic Games, Inc. All Rights Reserved.
+
+#pragma once
+
+#include "Delegates/Delegate.h"
+#include "ILauncherProfile.h"
+
+
+/**
+ * Delegate type for SProjectLauncher Profile running.
+ *
+ * The first parameter is the SProjectLauncher Profile that you want to run.
+ */
+DECLARE_DELEGATE_OneParam(FOnProfileRun, const ILauncherProfileRef&)
diff --git a/Engine/Source/Editor/MainFrame/Private/Menus/MainMenu.h b/Engine/Source/Editor/MainFrame/Private/Menus/MainMenu.h
--- a/Engine/Source/Editor/MainFrame/Private/Menus/MainMenu.h
+++ b/Engine/Source/Editor/MainFrame/Private/Menus/MainMenu.h
@@ -77,12 +77,12 @@ private:
 	static void RegisterNomadMainMenu();
 
 	/** 
-	* Opens the experimental project launcher tab.
-	* Remove this when it is is no longer experimental.
+	* Opens the legacy project launcher tab.
 	*/
+	UE_DEPRECATED(5.6, "invoke the tab directly")
 	static void OpenProjectLauncher()
 	{
-		FGlobalTabmanager::Get()->TryInvokeTab(FName(TEXT("ProjectLauncher")));
+		FGlobalTabmanager::Get()->TryInvokeTab(FName(TEXT("LegacyProjectLauncher")));
 	}
 
 	/**
diff --git a/Engine/Source/Editor/TurnkeySupport/Private/TurnkeySupportModule.cpp b/Engine/Source/Editor/TurnkeySupport/Private/TurnkeySupportModule.cpp
--- a/Engine/Source/Editor/TurnkeySupport/Private/TurnkeySupportModule.cpp
+++ b/Engine/Source/Editor/TurnkeySupport/Private/TurnkeySupportModule.cpp
@@ -58,7 +58,6 @@ IMPLEMENT_MODULE(FTurnkeySupportModuleEmpty, TurnkeySupport)
 #if WITH_EDITOR
 #include "LevelEditor.h"
 #include "Experimental/ZenServerInterface.h"
-#include "Settings/EditorExperimentalSettings.h"
 #endif
 
 #include "Misc/App.h"
@@ -228,14 +227,14 @@ public:
 		return FString::Printf(TEXT("-ReportFilename=\"%s\" -log=\"%s\""), *ReportFilename, *LogFilename);
 	}
 
-	static void OpenProjectLauncher()
+	static void OpenLegacyProjectLauncher()
 	{
-		FGlobalTabmanager::Get()->TryInvokeTab(FTabId("ProjectLauncher"));
+		FGlobalTabmanager::Get()->TryInvokeTab(FTabId("LegacyProjectLauncher"));
 	}
 
-	static void OpenCustomLaunchUI()
+	static void OpenProjectLauncher()
 	{
-		FGlobalTabmanager::Get()->TryInvokeTab(FTabId("CustomLaunchUI"));
+		FGlobalTabmanager::Get()->TryInvokeTab(FTabId("ProjectLauncher"));
 	}
 
 	static void OpenDeviceManager()
@@ -2368,35 +2367,27 @@ TSharedRef<SWidget> FTurnkeySupportModule::MakeTurnkeyMenuWidget() const
 		// options section
 		FToolMenuSection& OptionsSection = Menu->AddSection("TurnkeyOptions", LOCTEXT("TurnkeySection_Options", "Options and Settings"));
 		{
-#if WITH_EDITOR
 			if (IPluginManager::Get().FindEnabledPlugin(TEXT("CustomLaunchUI")) != nullptr)
 			{
 				OptionsSection.AddMenuEntry(
 					NAME_None,
-					LOCTEXT("OpenCustomLaunchProfileUI", "Custom Launch... (experimental)"),
-					LOCTEXT("OpenCustomLaunchProfileUI_ToolTip", "Configure custom launch profiles for advanced packaging, deploying and launching of your projects"),
+					LOCTEXT("OpenProjectLauncher", "Project Launcher..."),
+					LOCTEXT("OpenProjectLauncher_ToolTip", "Open the Project Launcher for advanced packaging, deploying and launching of your projects"),
 					FSlateIcon(FAppStyle::Get().GetStyleSetName(), "Launcher.TabIcon"),
-					FUIAction(
-						FExecuteAction::CreateStatic(&FTurnkeySupportCallbacks::OpenCustomLaunchUI), 
-						FCanExecuteAction(), 
-						FIsActionChecked(), 
-						FIsActionButtonVisible::CreateLambda([]() 
-						{
-							return GetDefault<UEditorExperimentalSettings>()->bShowNewLaunchProfileTools; 
-						}
-					)
-				));
+					FUIAction(FExecuteAction::CreateStatic(&FTurnkeySupportCallbacks::OpenProjectLauncher))
+				);
 			}
-#endif
 
 			OptionsSection.AddMenuEntry(
 				NAME_None,
-				LOCTEXT("OpenProjectLauncher", "Project Launcher..."),
-				LOCTEXT("OpenProjectLauncher_ToolTip", "Open the Project Launcher for advanced packaging, deploying and launching of your projects"),
+				LOCTEXT("OpenLegacyProjectLauncher", "Legacy Project Launcher..."),
+				LOCTEXT("OpenLegacyProjectLauncher_ToolTip", "Open the legacy Project Launcher for advanced packaging, deploying and launching of your projects"),
 				FSlateIcon(FAppStyle::Get().GetStyleSetName(), "Launcher.TabIcon"),
-				FUIAction(FExecuteAction::CreateStatic(&FTurnkeySupportCallbacks::OpenProjectLauncher))
+				FUIAction(FExecuteAction::CreateStatic(&FTurnkeySupportCallbacks::OpenLegacyProjectLauncher))
 			);
 
+
+
 			OptionsSection.AddMenuEntry(
 				NAME_None,
 				LOCTEXT("OpenDeviceManager", "Device Manager..."),
diff --git a/Engine/Source/Editor/UnrealEd/Private/EditorEngine.cpp b/Engine/Source/Editor/UnrealEd/Private/EditorEngine.cpp
--- a/Engine/Source/Editor/UnrealEd/Private/EditorEngine.cpp
+++ b/Engine/Source/Editor/UnrealEd/Private/EditorEngine.cpp
@@ -1481,7 +1481,7 @@ void UEditorEngine::LoadDefaultEditorModules()
 			TEXT("DeviceManager"),
 			TEXT("ProfilerClient"),
 			TEXT("SessionFrontend"),
-			TEXT("ProjectLauncher"),
+			TEXT("LegacyProjectLauncher"),
 			TEXT("SettingsEditor"),
 			TEXT("EditorSettingsViewer"),
 			TEXT("ProjectSettingsViewer"),
diff --git a/Engine/Source/Editor/UnrealEd/UnrealEd.Build.cs b/Engine/Source/Editor/UnrealEd/UnrealEd.Build.cs
--- a/Engine/Source/Editor/UnrealEd/UnrealEd.Build.cs
+++ b/Engine/Source/Editor/UnrealEd/UnrealEd.Build.cs
@@ -259,7 +259,7 @@ public class UnrealEd : ModuleRules
 			DynamicallyLoadedModuleNames.AddRange(
 				new string[] {
 					"SessionFrontend",
-					"ProjectLauncher",
+					"LegacyProjectLauncher",
 					"DeviceManager",
 					"ProjectTargetPlatformEditor",
 					"PListEditor",
diff --git a/Engine/Source/Programs/UnrealFrontend/Private/Commands/UserInterfaceCommand.cpp b/Engine/Source/Programs/UnrealFrontend/Private/Commands/UserInterfaceCommand.cpp
--- a/Engine/Source/Programs/UnrealFrontend/Private/Commands/UserInterfaceCommand.cpp
+++ b/Engine/Source/Programs/UnrealFrontend/Private/Commands/UserInterfaceCommand.cpp
@@ -63,7 +63,7 @@ void FUserInterfaceCommand::Run(  )
 #if STATS && UE_DEPRECATED_PROFILER_ENABLED
 	FModuleManager::Get().LoadModule("ProfilerClient");
 #endif
-	FModuleManager::Get().LoadModule("ProjectLauncher");
+	FModuleManager::Get().LoadModule("LegacyProjectLauncher");
 	FModuleManager::Get().LoadModule("SessionFrontend");
 	FModuleManager::Get().LoadModule("SettingsEditor");
 	IPluginManager::Get().LoadModulesForEnabledPlugins(ELoadingPhase::Default);
@@ -167,7 +167,7 @@ void FUserInterfaceCommand::InitializeSlateApplication( const FString& LayoutIni
 						->AddTab(FName("MessagingDebugger"), ETabState::ClosedTab)
 						->AddTab(FName("SessionFrontend"), ETabState::OpenedTab)
 						->AddTab(FName("ProjectLauncher"), ETabState::OpenedTab)
-						->AddTab(FName("CustomLaunchUI"), ETabState::OpenedTab)
+						->AddTab(FName("LegacyProjectLauncher"), ETabState::OpenedTab)
 				)
 		)
 		->AddArea
diff --git a/Engine/Source/Programs/UnrealFrontend/UnrealFrontend.Build.cs b/Engine/Source/Programs/UnrealFrontend/UnrealFrontend.Build.cs
--- a/Engine/Source/Programs/UnrealFrontend/UnrealFrontend.Build.cs
+++ b/Engine/Source/Programs/UnrealFrontend/UnrealFrontend.Build.cs
@@ -26,7 +26,7 @@ public class UnrealFrontend : ModuleRules
 				"LauncherServices",
 				"Messaging",
 				"OutputLog",
-				"ProjectLauncher",
+				"LegacyProjectLauncher",
 				"Projects",
 				"SessionFrontend",
 				"SessionServices",
