Rename clashing symbols from legacy project launcher to fix monolithic editor builds

#rnx

diff --git a/Engine/Source/Developer/LegacyProjectLauncher/Private/Models/LegacyProjectLauncherModel.h b/Engine/Source/Developer/LegacyProjectLauncher/Private/Models/LegacyProjectLauncherModel.h
new file mode 100644
--- /dev/null
+++ b/Engine/Source/Developer/LegacyProjectLauncher/Private/Models/LegacyProjectLauncherModel.h
@@ -0,0 +1,264 @@
+// Copyright Epic Games, Inc. All Rights Reserved.
+
+#pragma once
+
+#include "CoreTypes.h"
+#include "Containers/Array.h"
+#include "Delegates/Delegate.h"
+#include "GameProjectHelper.h"
+#include "ILauncher.h"
+#include "ILauncherProfileManager.h"
+#include "Misc/Paths.h"
+#include "Misc/ConfigCacheIni.h"
+#include "Misc/App.h"
+#include "Templates/SharedPointer.h"
+
+
+class FProjectLauncherModel;
+class SLegacyProjectLauncher;
+
+
+namespace ELauncherPanels
+{
+	/**
+	 * Enumerates available launcher panels.
+	 */
+	enum Type
+	{
+		NoTask = 0,
+		Launch,
+		ProfileEditor,
+		Progress,
+		End,
+	};
+}
+
+
+/**
+ * Delegate type for launcher profile selection changes.
+ *
+ * The first parameter is the newly selected profile (or nullptr if none was selected).
+ * The second parameter is the old selected profile (or nullptr if none was previous selected).
+ */
+DECLARE_MULTICAST_DELEGATE_TwoParams(FOnSelectedSProjectLauncherProfileChanged, const ILauncherProfilePtr&, const ILauncherProfilePtr&)
+
+
+/**
+ * Implements the data model for the session launcher.
+ */
+class FProjectLauncherModel
+{
+public:
+
+	/**
+	 * Creates and initializes a new instance.
+	 *
+	 * @param InDeviceProxyManager The device proxy manager to use.
+	 * @param InLauncher The launcher to use.
+	 * @param InProfileManager The profile manager to use.
+	 */
+	FProjectLauncherModel(const TSharedRef<ITargetDeviceProxyManager>& InDeviceProxyManager, const TSharedRef<ILauncher>& InLauncher, const TSharedRef<ILauncherProfileManager>& InProfileManager)
+		: DeviceProxyManager(InDeviceProxyManager)
+		, SLegacyProjectLauncher(InLauncher)
+		, ProfileManager(InProfileManager)
+	{
+		ProfileManager->OnProfileAdded().AddRaw(this, &FProjectLauncherModel::HandleProfileManagerProfileAdded);
+		ProfileManager->OnProfileRemoved().AddRaw(this, &FProjectLauncherModel::HandleProfileManagerProfileRemoved);
+
+		LoadConfig();
+	}
+
+	/** Destructor. */
+	~FProjectLauncherModel()
+	{
+		ProfileManager->OnProfileAdded().RemoveAll(this);
+		ProfileManager->OnProfileRemoved().RemoveAll(this);
+
+		SaveConfig();
+	}
+
+public:
+
+	/**
+	 * Gets the model's device proxy manager.
+	 *
+	 * @return Device proxy manager.
+	 */
+	const TSharedRef<ITargetDeviceProxyManager>& GetDeviceProxyManager() const
+	{
+		return DeviceProxyManager;
+	}
+
+	/**
+	 * Gets the model's launcher.
+	 *
+	 * @return SLegacyProjectLauncher.
+	 */
+	const TSharedRef<ILauncher>& GetSProjectLauncher() const
+	{
+		return SLegacyProjectLauncher;
+	}
+
+	/**
+	 * Get the model's profile manager.
+	 *
+	 * @return Profile manager.
+	 * @see GetProfiles
+	 */
+	const TSharedRef<ILauncherProfileManager>& GetProfileManager() const
+	{
+		return ProfileManager;
+	}
+
+	/**
+	 * Gets the active profile.
+	 *
+	 * @return The active profile.
+	 * @see GetProfiles, GetProfileManager, SelectProfile
+	 */
+	const TSharedPtr<ILauncherProfile>& GetSelectedProfile() const
+	{
+		return SelectedProfile;
+	}
+
+	/**
+	 * Sets the active profile.
+	 *
+	 * @param InProfile The profile to assign as active,, or nullptr to deselect all.
+	 * @see GetSelectedProfile
+	 */
+	void SelectProfile(const TSharedPtr<ILauncherProfile>& Profile)
+	{
+		if (!Profile.IsValid() || ProfileManager->GetAllProfiles().Contains(Profile))
+		{
+			if (Profile != SelectedProfile)
+			{
+				TSharedPtr<ILauncherProfile> PreviousProfile = SelectedProfile;
+				SelectedProfile = Profile;
+
+				ProfileSelectedDelegate.Broadcast(Profile, PreviousProfile);
+			}
+		}
+	}
+
+public:
+
+	/**
+	 * Returns a delegate to be invoked when the profile list has been modified.
+	 *
+	 * @return The delegate.
+	 * @see OnProfileSelected
+	 */
+	FSimpleMulticastDelegate& OnProfileListChanged()
+	{
+		return ProfileListChangedDelegate;
+	}
+
+	/**
+	 * Returns a delegate to be invoked when the selected profile changed.
+	 *
+	 * @return The delegate.
+	 * @see OnProfileListChanged
+	 */
+	FOnSelectedSProjectLauncherProfileChanged& OnProfileSelected()
+	{
+		return ProfileSelectedDelegate;
+	}
+
+protected:
+
+	/*
+	 * Load all profiles from disk.
+	 *
+	 * @see SaveConfig
+	 */
+	void LoadConfig()
+	{
+		// restore the previous project selection
+		FString ProjectPath;
+
+		if (FPaths::IsProjectFilePathSet())
+		{
+			ProjectPath = FPaths::GetProjectFilePath();
+		}
+		else if (FGameProjectHelper::IsGameAvailable(FApp::GetProjectName()))
+		{
+			ProjectPath = FPaths::RootDir() / FApp::GetProjectName() / FApp::GetProjectName() + TEXT(".uproject");
+		}
+		else if (GConfig != nullptr)
+		{
+			GConfig->GetString(TEXT("FProjectLauncherModel"), TEXT("SelectedProjectPath"), ProjectPath, GEngineIni);
+		}
+
+		ProfileManager->SetProjectPath(ProjectPath);
+	}
+
+	/*
+	 * Saves all profiles to disk.
+
+	 * @see LoadConfig
+	 */
+	void SaveConfig()
+	{
+		if (GConfig != nullptr && !FPaths::IsProjectFilePathSet() && !FGameProjectHelper::IsGameAvailable(FApp::GetProjectName()))
+		{
+			FString ProjectPath = ProfileManager->GetProjectPath();
+			GConfig->SetString(TEXT("FProjectLauncherModel"), TEXT("SelectedProjectPath"), *ProjectPath, GEngineIni);
+		}
+	}
+
+private:
+
+	/** Callback for when a profile was added to the profile manager. */
+	void HandleProfileManagerProfileAdded(const TSharedRef<ILauncherProfile>& Profile)
+	{
+		ProfileListChangedDelegate.Broadcast();
+
+		SelectProfile(Profile);
+	}
+
+	/** Callback for when a profile was removed from the profile manager. */
+	void HandleProfileManagerProfileRemoved(const TSharedRef<ILauncherProfile>& Profile)
+	{
+		ProfileListChangedDelegate.Broadcast();
+
+		if (Profile == SelectedProfile)
+		{
+			const TArray<TSharedPtr<ILauncherProfile>>& Profiles = ProfileManager->GetAllProfiles();
+
+			if (Profiles.Num() > 0)
+			{
+				SelectProfile(Profiles[0]);
+			}
+			else
+			{
+				SelectProfile(nullptr);
+			}
+		}
+	}
+
+private:
+
+	/** Pointer to the device proxy manager. */
+	TSharedRef<ITargetDeviceProxyManager> DeviceProxyManager;
+
+	/** Pointer to the launcher. */
+	TSharedRef<ILauncher> SLegacyProjectLauncher;
+
+	/** Pointer to the profile manager. */
+	TSharedRef<ILauncherProfileManager> ProfileManager;
+
+	/** Pointer to active profile. */
+	TSharedPtr<ILauncherProfile> SelectedProfile;
+
+private:
+
+	/** A delegate to be invoked when the project path has been modified. */
+	FSimpleMulticastDelegate ProjectPathChangedDelegate;
+
+	/** A delegate to be invoked when the profile list has been modified. */
+	FSimpleMulticastDelegate ProfileListChangedDelegate;
+
+	/** A delegate to be invoked when the selected profile changed. */
+	FOnSelectedSProjectLauncherProfileChanged ProfileSelectedDelegate;
+};
diff --git a/Engine/Source/Developer/LegacyProjectLauncher/Private/Models/ProjectLauncherModel.h b/Engine/Source/Developer/LegacyProjectLauncher/Private/Models/ProjectLauncherModel.h
deleted file mode 100644
diff --git a/Engine/Source/Developer/LegacyProjectLauncher/Private/ProjectLauncherModule.cpp b/Engine/Source/Developer/LegacyProjectLauncher/Private/ProjectLauncherModule.cpp
--- a/Engine/Source/Developer/LegacyProjectLauncher/Private/ProjectLauncherModule.cpp
+++ b/Engine/Source/Developer/LegacyProjectLauncher/Private/ProjectLauncherModule.cpp
@@ -13,8 +13,8 @@
 #include "Widgets/DeclarativeSyntaxSupport.h"
 #include "Widgets/Docking/SDockTab.h"
 
-#include "Models/ProjectLauncherModel.h"
-#include "Widgets/SProjectLauncher.h"
+#include "Models/LegacyProjectLauncherModel.h"
+#include "Widgets/SLegacyProjectLauncher.h"
 
 
 static const FName ProjectLauncherTabName("LegacyProjectLauncher");
@@ -66,7 +66,7 @@ private:
 			ProjectLauncherServicesModule.GetProfileManager()
 		));
 
-		DockTab->SetContent(SNew(SProjectLauncher, DockTab, SpawnTabArgs.GetOwnerWindow(), Model));
+		DockTab->SetContent(SNew(SLegacyProjectLauncher, DockTab, SpawnTabArgs.GetOwnerWindow(), Model));
 
 		return DockTab;
 	}
diff --git a/Engine/Source/Developer/LegacyProjectLauncher/Private/Widgets/Archive/SProjectLauncherArchivePage.h b/Engine/Source/Developer/LegacyProjectLauncher/Private/Widgets/Archive/SProjectLauncherArchivePage.h
--- a/Engine/Source/Developer/LegacyProjectLauncher/Private/Widgets/Archive/SProjectLauncherArchivePage.h
+++ b/Engine/Source/Developer/LegacyProjectLauncher/Private/Widgets/Archive/SProjectLauncherArchivePage.h
@@ -7,7 +7,7 @@
 #include "Widgets/DeclarativeSyntaxSupport.h"
 #include "Input/Reply.h"
 #include "Widgets/SCompoundWidget.h"
-#include "Models/ProjectLauncherModel.h"
+#include "Models/LegacyProjectLauncherModel.h"
 
 enum class ECheckBoxState : uint8;
 
diff --git a/Engine/Source/Developer/LegacyProjectLauncher/Private/Widgets/Build/SProjectLauncherBuildPage.h b/Engine/Source/Developer/LegacyProjectLauncher/Private/Widgets/Build/SProjectLauncherBuildPage.h
--- a/Engine/Source/Developer/LegacyProjectLauncher/Private/Widgets/Build/SProjectLauncherBuildPage.h
+++ b/Engine/Source/Developer/LegacyProjectLauncher/Private/Widgets/Build/SProjectLauncherBuildPage.h
@@ -7,7 +7,7 @@
 #include "Widgets/DeclarativeSyntaxSupport.h"
 #include "Input/Reply.h"
 #include "Widgets/SCompoundWidget.h"
-#include "Models/ProjectLauncherModel.h"
+#include "Models/LegacyProjectLauncherModel.h"
 #include "ILauncherProfile.h"
 
 class Error;
diff --git a/Engine/Source/Developer/LegacyProjectLauncher/Private/Widgets/Build/SProjectLauncherBuildTaskSettings.h b/Engine/Source/Developer/LegacyProjectLauncher/Private/Widgets/Build/SProjectLauncherBuildTaskSettings.h
--- a/Engine/Source/Developer/LegacyProjectLauncher/Private/Widgets/Build/SProjectLauncherBuildTaskSettings.h
+++ b/Engine/Source/Developer/LegacyProjectLauncher/Private/Widgets/Build/SProjectLauncherBuildTaskSettings.h
@@ -5,7 +5,7 @@
 #include "CoreMinimal.h"
 #include "Widgets/DeclarativeSyntaxSupport.h"
 #include "Widgets/SCompoundWidget.h"
-#include "Models/ProjectLauncherModel.h"
+#include "Models/LegacyProjectLauncherModel.h"
 
 /**
  * Implements the launcher settings widget.
diff --git a/Engine/Source/Developer/LegacyProjectLauncher/Private/Widgets/Cook/SProjectLauncherCookByTheBookSettings.h b/Engine/Source/Developer/LegacyProjectLauncher/Private/Widgets/Cook/SProjectLauncherCookByTheBookSettings.h
--- a/Engine/Source/Developer/LegacyProjectLauncher/Private/Widgets/Cook/SProjectLauncherCookByTheBookSettings.h
+++ b/Engine/Source/Developer/LegacyProjectLauncher/Private/Widgets/Cook/SProjectLauncherCookByTheBookSettings.h
@@ -8,7 +8,7 @@
 #include "Input/Reply.h"
 #include "Widgets/SWidget.h"
 #include "Widgets/SCompoundWidget.h"
-#include "Models/ProjectLauncherModel.h"
+#include "Models/LegacyProjectLauncherModel.h"
 #include "Widgets/Views/STableViewBase.h"
 #include "Widgets/Views/STableRow.h"
 
diff --git a/Engine/Source/Developer/LegacyProjectLauncher/Private/Widgets/Cook/SProjectLauncherCookOnTheFlySettings.h b/Engine/Source/Developer/LegacyProjectLauncher/Private/Widgets/Cook/SProjectLauncherCookOnTheFlySettings.h
--- a/Engine/Source/Developer/LegacyProjectLauncher/Private/Widgets/Cook/SProjectLauncherCookOnTheFlySettings.h
+++ b/Engine/Source/Developer/LegacyProjectLauncher/Private/Widgets/Cook/SProjectLauncherCookOnTheFlySettings.h
@@ -6,7 +6,7 @@
 #include "Layout/Visibility.h"
 #include "Widgets/DeclarativeSyntaxSupport.h"
 #include "Widgets/SCompoundWidget.h"
-#include "Models/ProjectLauncherModel.h"
+#include "Models/LegacyProjectLauncherModel.h"
 
 class Error;
 enum class ECheckBoxState : uint8;
diff --git a/Engine/Source/Developer/LegacyProjectLauncher/Private/Widgets/Cook/SProjectLauncherCookPage.h b/Engine/Source/Developer/LegacyProjectLauncher/Private/Widgets/Cook/SProjectLauncherCookPage.h
--- a/Engine/Source/Developer/LegacyProjectLauncher/Private/Widgets/Cook/SProjectLauncherCookPage.h
+++ b/Engine/Source/Developer/LegacyProjectLauncher/Private/Widgets/Cook/SProjectLauncherCookPage.h
@@ -6,7 +6,7 @@
 #include "Layout/Visibility.h"
 #include "Widgets/DeclarativeSyntaxSupport.h"
 #include "Widgets/SCompoundWidget.h"
-#include "Models/ProjectLauncherModel.h"
+#include "Models/LegacyProjectLauncherModel.h"
 
 class SEditableTextBox;
 
diff --git a/Engine/Source/Developer/LegacyProjectLauncher/Private/Widgets/Cook/SProjectLauncherCookedPlatforms.h b/Engine/Source/Developer/LegacyProjectLauncher/Private/Widgets/Cook/SProjectLauncherCookedPlatforms.h
--- a/Engine/Source/Developer/LegacyProjectLauncher/Private/Widgets/Cook/SProjectLauncherCookedPlatforms.h
+++ b/Engine/Source/Developer/LegacyProjectLauncher/Private/Widgets/Cook/SProjectLauncherCookedPlatforms.h
@@ -7,7 +7,7 @@
 #include "Widgets/DeclarativeSyntaxSupport.h"
 #include "Styling/SlateColor.h"
 #include "Widgets/SCompoundWidget.h"
-#include "Models/ProjectLauncherModel.h"
+#include "Models/LegacyProjectLauncherModel.h"
 #include "Widgets/Views/STableViewBase.h"
 #include "Widgets/Views/STableRow.h"
 #include "Styling/AppStyle.h"
diff --git a/Engine/Source/Developer/LegacyProjectLauncher/Private/Widgets/Cook/SProjectLauncherCultureListRow.h b/Engine/Source/Developer/LegacyProjectLauncher/Private/Widgets/Cook/SProjectLauncherCultureListRow.h
--- a/Engine/Source/Developer/LegacyProjectLauncher/Private/Widgets/Cook/SProjectLauncherCultureListRow.h
+++ b/Engine/Source/Developer/LegacyProjectLauncher/Private/Widgets/Cook/SProjectLauncherCultureListRow.h
@@ -8,7 +8,7 @@
 #include "Widgets/DeclarativeSyntaxSupport.h"
 #include "Widgets/SWidget.h"
 #include "Layout/Margin.h"
-#include "Models/ProjectLauncherModel.h"
+#include "Models/LegacyProjectLauncherModel.h"
 #include "Styling/SlateTypes.h"
 #include "Widgets/Views/STableRow.h"
 #include "Widgets/Text/STextBlock.h"
diff --git a/Engine/Source/Developer/LegacyProjectLauncher/Private/Widgets/Cook/SProjectLauncherMapListRow.h b/Engine/Source/Developer/LegacyProjectLauncher/Private/Widgets/Cook/SProjectLauncherMapListRow.h
--- a/Engine/Source/Developer/LegacyProjectLauncher/Private/Widgets/Cook/SProjectLauncherMapListRow.h
+++ b/Engine/Source/Developer/LegacyProjectLauncher/Private/Widgets/Cook/SProjectLauncherMapListRow.h
@@ -8,7 +8,7 @@
 #include "Widgets/DeclarativeSyntaxSupport.h"
 #include "Widgets/SWidget.h"
 #include "Layout/Margin.h"
-#include "Models/ProjectLauncherModel.h"
+#include "Models/LegacyProjectLauncherModel.h"
 #include "Styling/SlateTypes.h"
 #include "Widgets/Views/STableRow.h"
 #include "Widgets/Text/STextBlock.h"
diff --git a/Engine/Source/Developer/LegacyProjectLauncher/Private/Widgets/Cook/SProjectLauncherPlatformListRow.h b/Engine/Source/Developer/LegacyProjectLauncher/Private/Widgets/Cook/SProjectLauncherPlatformListRow.h
--- a/Engine/Source/Developer/LegacyProjectLauncher/Private/Widgets/Cook/SProjectLauncherPlatformListRow.h
+++ b/Engine/Source/Developer/LegacyProjectLauncher/Private/Widgets/Cook/SProjectLauncherPlatformListRow.h
@@ -8,7 +8,7 @@
 #include "Widgets/DeclarativeSyntaxSupport.h"
 #include "Widgets/SWidget.h"
 #include "Layout/Margin.h"
-#include "Models/ProjectLauncherModel.h"
+#include "Models/LegacyProjectLauncherModel.h"
 #include "Styling/SlateTypes.h"
 #include "Widgets/Views/STableRow.h"
 #include "Widgets/Text/STextBlock.h"
diff --git a/Engine/Source/Developer/LegacyProjectLauncher/Private/Widgets/Cook/SProjectLauncherSimpleCookPage.h b/Engine/Source/Developer/LegacyProjectLauncher/Private/Widgets/Cook/SProjectLauncherSimpleCookPage.h
--- a/Engine/Source/Developer/LegacyProjectLauncher/Private/Widgets/Cook/SProjectLauncherSimpleCookPage.h
+++ b/Engine/Source/Developer/LegacyProjectLauncher/Private/Widgets/Cook/SProjectLauncherSimpleCookPage.h
@@ -6,7 +6,7 @@
 #include "Layout/Visibility.h"
 #include "Widgets/DeclarativeSyntaxSupport.h"
 #include "Widgets/SCompoundWidget.h"
-#include "Models/ProjectLauncherModel.h"
+#include "Models/LegacyProjectLauncherModel.h"
 
 class SEditableTextBox;
 
diff --git a/Engine/Source/Developer/LegacyProjectLauncher/Private/Widgets/Deploy/SProjectLauncherDeployFileServerSettings.cpp b/Engine/Source/Developer/LegacyProjectLauncher/Private/Widgets/Deploy/SProjectLauncherDeployFileServerSettings.cpp
--- a/Engine/Source/Developer/LegacyProjectLauncher/Private/Widgets/Deploy/SProjectLauncherDeployFileServerSettings.cpp
+++ b/Engine/Source/Developer/LegacyProjectLauncher/Private/Widgets/Deploy/SProjectLauncherDeployFileServerSettings.cpp
@@ -9,7 +9,7 @@
 #include "Widgets/Layout/SBorder.h"
 #include "Widgets/Text/STextBlock.h"
 
-#include "Models/ProjectLauncherModel.h"
+#include "Models/LegacyProjectLauncherModel.h"
 #include "Widgets/Deploy/SProjectLauncherDeployTargets.h"
 #include "Widgets/Layout/SExpandableArea.h"
 
diff --git a/Engine/Source/Developer/LegacyProjectLauncher/Private/Widgets/Deploy/SProjectLauncherDeployPage.h b/Engine/Source/Developer/LegacyProjectLauncher/Private/Widgets/Deploy/SProjectLauncherDeployPage.h
--- a/Engine/Source/Developer/LegacyProjectLauncher/Private/Widgets/Deploy/SProjectLauncherDeployPage.h
+++ b/Engine/Source/Developer/LegacyProjectLauncher/Private/Widgets/Deploy/SProjectLauncherDeployPage.h
@@ -6,7 +6,7 @@
 #include "Layout/Visibility.h"
 #include "Widgets/DeclarativeSyntaxSupport.h"
 #include "Widgets/SCompoundWidget.h"
-#include "Models/ProjectLauncherModel.h"
+#include "Models/LegacyProjectLauncherModel.h"
 
 class Error;
 
diff --git a/Engine/Source/Developer/LegacyProjectLauncher/Private/Widgets/Deploy/SProjectLauncherDeployRepositorySettings.h b/Engine/Source/Developer/LegacyProjectLauncher/Private/Widgets/Deploy/SProjectLauncherDeployRepositorySettings.h
--- a/Engine/Source/Developer/LegacyProjectLauncher/Private/Widgets/Deploy/SProjectLauncherDeployRepositorySettings.h
+++ b/Engine/Source/Developer/LegacyProjectLauncher/Private/Widgets/Deploy/SProjectLauncherDeployRepositorySettings.h
@@ -6,7 +6,7 @@
 #include "Widgets/DeclarativeSyntaxSupport.h"
 #include "Input/Reply.h"
 #include "Widgets/SCompoundWidget.h"
-#include "Models/ProjectLauncherModel.h"
+#include "Models/LegacyProjectLauncherModel.h"
 
 class SEditableTextBox;
 
diff --git a/Engine/Source/Developer/LegacyProjectLauncher/Private/Widgets/Deploy/SProjectLauncherDeployTargets.h b/Engine/Source/Developer/LegacyProjectLauncher/Private/Widgets/Deploy/SProjectLauncherDeployTargets.h
--- a/Engine/Source/Developer/LegacyProjectLauncher/Private/Widgets/Deploy/SProjectLauncherDeployTargets.h
+++ b/Engine/Source/Developer/LegacyProjectLauncher/Private/Widgets/Deploy/SProjectLauncherDeployTargets.h
@@ -7,7 +7,7 @@
 #include "Widgets/DeclarativeSyntaxSupport.h"
 #include "Widgets/SWidget.h"
 #include "Widgets/SCompoundWidget.h"
-#include "Models/ProjectLauncherModel.h"
+#include "Models/LegacyProjectLauncherModel.h"
 #include "Widgets/Views/STableViewBase.h"
 #include "Widgets/Views/STableRow.h"
 #include "Widgets/Input/SComboBox.h"
diff --git a/Engine/Source/Developer/LegacyProjectLauncher/Private/Widgets/Deploy/SProjectLauncherDeployTaskSettings.h b/Engine/Source/Developer/LegacyProjectLauncher/Private/Widgets/Deploy/SProjectLauncherDeployTaskSettings.h
--- a/Engine/Source/Developer/LegacyProjectLauncher/Private/Widgets/Deploy/SProjectLauncherDeployTaskSettings.h
+++ b/Engine/Source/Developer/LegacyProjectLauncher/Private/Widgets/Deploy/SProjectLauncherDeployTaskSettings.h
@@ -5,7 +5,7 @@
 #include "CoreMinimal.h"
 #include "Widgets/DeclarativeSyntaxSupport.h"
 #include "Widgets/SCompoundWidget.h"
-#include "Models/ProjectLauncherModel.h"
+#include "Models/LegacyProjectLauncherModel.h"
 
 /**
  * Implements the launcher settings widget.
diff --git a/Engine/Source/Developer/LegacyProjectLauncher/Private/Widgets/Deploy/SProjectLauncherDeployToDeviceSettings.h b/Engine/Source/Developer/LegacyProjectLauncher/Private/Widgets/Deploy/SProjectLauncherDeployToDeviceSettings.h
--- a/Engine/Source/Developer/LegacyProjectLauncher/Private/Widgets/Deploy/SProjectLauncherDeployToDeviceSettings.h
+++ b/Engine/Source/Developer/LegacyProjectLauncher/Private/Widgets/Deploy/SProjectLauncherDeployToDeviceSettings.h
@@ -6,7 +6,7 @@
 #include "Layout/Visibility.h"
 #include "Widgets/DeclarativeSyntaxSupport.h"
 #include "Widgets/SCompoundWidget.h"
-#include "Models/ProjectLauncherModel.h"
+#include "Models/LegacyProjectLauncherModel.h"
 
 enum class ECheckBoxState : uint8;
 
diff --git a/Engine/Source/Developer/LegacyProjectLauncher/Private/Widgets/Deploy/SProjectLauncherSimpleDeviceListRow.h b/Engine/Source/Developer/LegacyProjectLauncher/Private/Widgets/Deploy/SProjectLauncherSimpleDeviceListRow.h
--- a/Engine/Source/Developer/LegacyProjectLauncher/Private/Widgets/Deploy/SProjectLauncherSimpleDeviceListRow.h
+++ b/Engine/Source/Developer/LegacyProjectLauncher/Private/Widgets/Deploy/SProjectLauncherSimpleDeviceListRow.h
@@ -14,7 +14,7 @@
 #include "Widgets/SBoxPanel.h"
 #include "Widgets/Layout/SBox.h"
 #include "Widgets/Views/STableRow.h"
-#include "Models/ProjectLauncherModel.h"
+#include "Models/LegacyProjectLauncherModel.h"
 #include "Widgets/Shared/ProjectLauncherDelegates.h"
 #include "Widgets/Text/STextBlock.h"
 #include "Widgets/Views/SListView.h"
diff --git a/Engine/Source/Developer/LegacyProjectLauncher/Private/Widgets/Deploy/SProjectLauncherSimpleDeviceListView.cpp b/Engine/Source/Developer/LegacyProjectLauncher/Private/Widgets/Deploy/SProjectLauncherSimpleDeviceListView.cpp
--- a/Engine/Source/Developer/LegacyProjectLauncher/Private/Widgets/Deploy/SProjectLauncherSimpleDeviceListView.cpp
+++ b/Engine/Source/Developer/LegacyProjectLauncher/Private/Widgets/Deploy/SProjectLauncherSimpleDeviceListView.cpp
@@ -12,7 +12,7 @@
 #include "Widgets/Views/STableRow.h"
 #include "Widgets/Views/STableViewBase.h"
 
-#include "Models/ProjectLauncherModel.h"
+#include "Models/LegacyProjectLauncherModel.h"
 #include "Widgets/Deploy/SProjectLauncherSimpleDeviceListRow.h"
 
 
diff --git a/Engine/Source/Developer/LegacyProjectLauncher/Private/Widgets/Launch/SProjectLauncherLaunchCustomRoles.h b/Engine/Source/Developer/LegacyProjectLauncher/Private/Widgets/Launch/SProjectLauncherLaunchCustomRoles.h
--- a/Engine/Source/Developer/LegacyProjectLauncher/Private/Widgets/Launch/SProjectLauncherLaunchCustomRoles.h
+++ b/Engine/Source/Developer/LegacyProjectLauncher/Private/Widgets/Launch/SProjectLauncherLaunchCustomRoles.h
@@ -5,7 +5,7 @@
 #include "CoreMinimal.h"
 #include "Widgets/DeclarativeSyntaxSupport.h"
 #include "Widgets/SCompoundWidget.h"
-#include "Models/ProjectLauncherModel.h"
+#include "Models/LegacyProjectLauncherModel.h"
 
 /**
  * Implements the settings panel for launching with custom roles.
diff --git a/Engine/Source/Developer/LegacyProjectLauncher/Private/Widgets/Launch/SProjectLauncherLaunchPage.h b/Engine/Source/Developer/LegacyProjectLauncher/Private/Widgets/Launch/SProjectLauncherLaunchPage.h
--- a/Engine/Source/Developer/LegacyProjectLauncher/Private/Widgets/Launch/SProjectLauncherLaunchPage.h
+++ b/Engine/Source/Developer/LegacyProjectLauncher/Private/Widgets/Launch/SProjectLauncherLaunchPage.h
@@ -6,7 +6,7 @@
 #include "Layout/Visibility.h"
 #include "Widgets/DeclarativeSyntaxSupport.h"
 #include "Widgets/SCompoundWidget.h"
-#include "Models/ProjectLauncherModel.h"
+#include "Models/LegacyProjectLauncherModel.h"
 
 class Error;
 class SProjectLauncherLaunchRoleEditor;
diff --git a/Engine/Source/Developer/LegacyProjectLauncher/Private/Widgets/Launch/SProjectLauncherLaunchTaskSettings.h b/Engine/Source/Developer/LegacyProjectLauncher/Private/Widgets/Launch/SProjectLauncherLaunchTaskSettings.h
--- a/Engine/Source/Developer/LegacyProjectLauncher/Private/Widgets/Launch/SProjectLauncherLaunchTaskSettings.h
+++ b/Engine/Source/Developer/LegacyProjectLauncher/Private/Widgets/Launch/SProjectLauncherLaunchTaskSettings.h
@@ -5,7 +5,7 @@
 #include "CoreMinimal.h"
 #include "Widgets/DeclarativeSyntaxSupport.h"
 #include "Widgets/SCompoundWidget.h"
-#include "Models/ProjectLauncherModel.h"
+#include "Models/LegacyProjectLauncherModel.h"
 
 /**
  * Implements the launcher settings widget.
diff --git a/Engine/Source/Developer/LegacyProjectLauncher/Private/Widgets/Package/SProjectLauncherPackagePage.h b/Engine/Source/Developer/LegacyProjectLauncher/Private/Widgets/Package/SProjectLauncherPackagePage.h
--- a/Engine/Source/Developer/LegacyProjectLauncher/Private/Widgets/Package/SProjectLauncherPackagePage.h
+++ b/Engine/Source/Developer/LegacyProjectLauncher/Private/Widgets/Package/SProjectLauncherPackagePage.h
@@ -6,7 +6,7 @@
 #include "Layout/Visibility.h"
 #include "Widgets/DeclarativeSyntaxSupport.h"
 #include "Widgets/SCompoundWidget.h"
-#include "Models/ProjectLauncherModel.h"
+#include "Models/LegacyProjectLauncherModel.h"
 
 class SProjectLauncherPackagingSettings;
 
diff --git a/Engine/Source/Developer/LegacyProjectLauncher/Private/Widgets/Package/SProjectLauncherPackagingSettings.h b/Engine/Source/Developer/LegacyProjectLauncher/Private/Widgets/Package/SProjectLauncherPackagingSettings.h
--- a/Engine/Source/Developer/LegacyProjectLauncher/Private/Widgets/Package/SProjectLauncherPackagingSettings.h
+++ b/Engine/Source/Developer/LegacyProjectLauncher/Private/Widgets/Package/SProjectLauncherPackagingSettings.h
@@ -6,7 +6,7 @@
 #include "Widgets/DeclarativeSyntaxSupport.h"
 #include "Input/Reply.h"
 #include "Widgets/SCompoundWidget.h"
-#include "Models/ProjectLauncherModel.h"
+#include "Models/LegacyProjectLauncherModel.h"
 
 class SEditableTextBox;
 enum class ECheckBoxState : uint8;
diff --git a/Engine/Source/Developer/LegacyProjectLauncher/Private/Widgets/Preview/SProjectLauncherPreviewPage.h b/Engine/Source/Developer/LegacyProjectLauncher/Private/Widgets/Preview/SProjectLauncherPreviewPage.h
--- a/Engine/Source/Developer/LegacyProjectLauncher/Private/Widgets/Preview/SProjectLauncherPreviewPage.h
+++ b/Engine/Source/Developer/LegacyProjectLauncher/Private/Widgets/Preview/SProjectLauncherPreviewPage.h
@@ -10,7 +10,7 @@
 #include "Widgets/SCompoundWidget.h"
 #include "Widgets/Views/STableViewBase.h"
 #include "Widgets/Views/STableRow.h"
-#include "Models/ProjectLauncherModel.h"
+#include "Models/LegacyProjectLauncherModel.h"
 
 class ITargetDeviceProxy;
 
diff --git a/Engine/Source/Developer/LegacyProjectLauncher/Private/Widgets/Profile/SProjectLauncherProfileListView.cpp b/Engine/Source/Developer/LegacyProjectLauncher/Private/Widgets/Profile/SProjectLauncherProfileListView.cpp
--- a/Engine/Source/Developer/LegacyProjectLauncher/Private/Widgets/Profile/SProjectLauncherProfileListView.cpp
+++ b/Engine/Source/Developer/LegacyProjectLauncher/Private/Widgets/Profile/SProjectLauncherProfileListView.cpp
@@ -9,7 +9,7 @@
 #include "Widgets/Layout/SScrollBorder.h"
 
 #include "Models/ProjectLauncherCommands.h"
-#include "Models/ProjectLauncherModel.h"
+#include "Models/LegacyProjectLauncherModel.h"
 #include "Widgets/Profile/SProjectLauncherProfileListRow.h"
 
 
diff --git a/Engine/Source/Developer/LegacyProjectLauncher/Private/Widgets/Project/SProjectLauncherProjectPage.h b/Engine/Source/Developer/LegacyProjectLauncher/Private/Widgets/Project/SProjectLauncherProjectPage.h
--- a/Engine/Source/Developer/LegacyProjectLauncher/Private/Widgets/Project/SProjectLauncherProjectPage.h
+++ b/Engine/Source/Developer/LegacyProjectLauncher/Private/Widgets/Project/SProjectLauncherProjectPage.h
@@ -5,7 +5,7 @@
 #include "CoreMinimal.h"
 #include "Widgets/DeclarativeSyntaxSupport.h"
 #include "Widgets/SCompoundWidget.h"
-#include "Models/ProjectLauncherModel.h"
+#include "Models/LegacyProjectLauncherModel.h"
 
 /**
  * Implements the profile page for the session launcher wizard.
diff --git a/Engine/Source/Developer/LegacyProjectLauncher/Private/Widgets/Project/SProjectLauncherProjectPicker.h b/Engine/Source/Developer/LegacyProjectLauncher/Private/Widgets/Project/SProjectLauncherProjectPicker.h
--- a/Engine/Source/Developer/LegacyProjectLauncher/Private/Widgets/Project/SProjectLauncherProjectPicker.h
+++ b/Engine/Source/Developer/LegacyProjectLauncher/Private/Widgets/Project/SProjectLauncherProjectPicker.h
@@ -8,7 +8,7 @@
 #include "Widgets/DeclarativeSyntaxSupport.h"
 #include "Widgets/SWidget.h"
 #include "Widgets/SCompoundWidget.h"
-#include "Models/ProjectLauncherModel.h"
+#include "Models/LegacyProjectLauncherModel.h"
 
 class Error;
 
diff --git a/Engine/Source/Developer/LegacyProjectLauncher/Private/Widgets/SLegacyProjectLauncher.cpp b/Engine/Source/Developer/LegacyProjectLauncher/Private/Widgets/SLegacyProjectLauncher.cpp
new file mode 100644
--- /dev/null
+++ b/Engine/Source/Developer/LegacyProjectLauncher/Private/Widgets/SLegacyProjectLauncher.cpp
@@ -0,0 +1,437 @@
+// Copyright Epic Games, Inc. All Rights Reserved.
+
+#include "Widgets/SLegacyProjectLauncher.h"
+#include "Framework/Commands/UIAction.h"
+#include "Framework/Docking/TabManager.h"
+#include "Framework/MultiBox/MultiBoxBuilder.h"
+#include "SlateOptMacros.h"
+#include "Styling/SlateTypes.h"
+#include "Widgets/SBoxPanel.h"
+#include "Widgets/Images/SImage.h"
+#include "Widgets/Input/SButton.h"
+#include "Widgets/Input/SCheckBox.h"
+#include "Widgets/Input/SComboButton.h"
+#include "Widgets/Layout/SBorder.h"
+#include "Widgets/Layout/SSplitter.h"
+#include "Widgets/Layout/SWidgetSwitcher.h"
+#include "Widgets/Text/STextBlock.h"
+
+#include "Models/ProjectLauncherCommands.h"
+#include "Widgets/Deploy/SProjectLauncherSimpleDeviceListView.h"
+#include "Widgets/Profile/SProjectLauncherProfileListView.h"
+#include "Widgets/Progress/SProjectLauncherProgress.h"
+#include "Widgets/Project/SProjectLauncherProjectPicker.h"
+#include "Widgets/Shared/SProjectLauncherBuildTargetSelector.h"
+#include "Widgets/Settings/SProjectLauncherSettings.h"
+#include "Styling/AppStyle.h"
+#include "Styling/StyleColors.h"
+#include "SPositiveActionButton.h"
+
+#define LOCTEXT_NAMESPACE "SLegacyProjectLauncher"
+
+
+/* SLegacyProjectLauncher structors
+ *****************************************************************************/
+
+SLegacyProjectLauncher::SLegacyProjectLauncher()
+	:bAdvanced(false)
+{
+	if (GConfig != NULL)
+	{
+		GConfig->GetBool(TEXT("FProjectLauncher"), TEXT("AdvancedMode"), bAdvanced, GEngineIni);
+	}
+}
+
+
+SLegacyProjectLauncher::~SLegacyProjectLauncher()
+{
+	if (GConfig != NULL)
+	{
+		GConfig->SetBool(TEXT("FProjectLauncher"), TEXT("AdvancedMode"), bAdvanced, GEngineIni);
+	}
+
+	if (LauncherWorker.IsValid())
+	{
+		LauncherWorker->Cancel();
+		FPlatformProcess::Sleep(0.5f);
+	}
+}
+
+
+/* SLegacyProjectLauncher interface
+ *****************************************************************************/
+
+BEGIN_SLATE_FUNCTION_BUILD_OPTIMIZATION
+void SLegacyProjectLauncher::Construct(const FArguments& InArgs, const TSharedRef<SDockTab>& ConstructUnderMajorTab, const TSharedPtr<SWindow>& ConstructUnderWindow, const TSharedRef<FProjectLauncherModel>& InModel)
+{
+	FProjectLauncherCommands::Register();
+
+	Model = InModel;
+
+	// create & initialize main menu bar
+	TSharedRef<FWorkspaceItem> RootMenuGroup = FWorkspaceItem::NewGroup(LOCTEXT("RootMenuGroup", "Root"));
+
+	FMenuBarBuilder MenuBarBuilder = FMenuBarBuilder(TSharedPtr<FUICommandList>());
+	MenuBarBuilder.AddPullDownMenu(
+		LOCTEXT("WindowMenuLabel", "Window"),
+		FText::GetEmpty(),
+		FNewMenuDelegate::CreateStatic(&SLegacyProjectLauncher::FillWindowMenu, RootMenuGroup),
+		"Window"
+	);
+
+	ChildSlot
+	[
+		SAssignNew(WidgetSwitcher, SWidgetSwitcher)
+		.WidgetIndex((int32)ELauncherPanels::Launch)
+		
+		// Empty Panel
+		+ SWidgetSwitcher::Slot()
+		[
+			SNew(SBorder)
+		]
+
+		// SLegacyProjectLauncher Panel
+		+ SWidgetSwitcher::Slot()
+		[
+			SNew(SSplitter)
+			.Orientation(Orient_Vertical)
+
+			// Simple SLegacyProjectLauncher
+			+ SSplitter::Slot()
+			.Value(0.55f)
+			[
+				SNew(SVerticalBox)
+
+				+ SVerticalBox::Slot()
+				.AutoHeight()
+				[
+					SNew(SBorder)
+					.BorderImage(FAppStyle::Get().GetBrush("Brushes.Panel"))
+					[
+						SNew(SHorizontalBox)
+
+						// Project Bar
+						+ SHorizontalBox::Slot()
+						.AutoWidth()
+						.Padding(0, 4)
+						[
+							SNew(SProjectLauncherProjectPicker, InModel)
+						]
+
+						// Build Target
+						+ SHorizontalBox::Slot()
+						.AutoWidth()
+						.Padding(0,4)
+						[
+							SNew(SProjectLauncherBuildTargetSelector, InModel)
+							.UseProfile(false)
+						]
+
+						// Advanced Button
+						+ SHorizontalBox::Slot()
+						.AutoWidth()
+						.VAlign(VAlign_Center)
+						.Padding(6, 0, 0, 0)
+						[
+							
+							SNew(SCheckBox)
+							.Style(FAppStyle::Get(), "ToggleButtonCheckbox")
+							.IsFocusable(true)
+							.ToolTipText(LOCTEXT("ToggleAdvancedOptionsToolTipText", "Toggles Advanced Options"))
+							.OnCheckStateChanged(this, &SLegacyProjectLauncher::OnAdvancedChanged)
+							.IsChecked(this, &SLegacyProjectLauncher::OnIsAdvanced)
+							[
+								SNew(SHorizontalBox)
+								// Icon
+								+ SHorizontalBox::Slot()
+								.VAlign(VAlign_Center)
+								.AutoWidth()
+								[
+									SNew(SImage)
+									.Image(this, &SLegacyProjectLauncher::GetAdvancedToggleBrush)
+									.DesiredSizeOverride(FVector2D(16, 16))
+									.ColorAndOpacity(FSlateColor::UseForeground())
+								]
+
+								// Text
+								+ SHorizontalBox::Slot()
+								.AutoWidth()
+								.VAlign(VAlign_Center)
+								.Padding(4, 0, 4, 0)
+								[
+									SNew(STextBlock)
+									.Text(LOCTEXT("AdvancedButton", "Show Advanced"))
+								]
+							]
+						]
+					]
+				]
+
+				+ SVerticalBox::Slot()
+				.FillHeight(1)
+				.Padding(0)
+				[
+					SNew(SBorder)
+					.BorderImage(FAppStyle::Get().GetBrush("Brushes.Recessed"))
+					[
+						SAssignNew(LaunchList, SProjectLauncherSimpleDeviceListView, InModel)
+						.OnProfileRun(this, &SLegacyProjectLauncher::OnProfileRun)
+						.IsAdvanced(this, &SLegacyProjectLauncher::GetIsAdvanced)
+					]
+				]
+			]
+
+			+ SSplitter::Slot()
+			.Value(0.45f)
+			[
+				SNew(SBorder)
+				.Padding(0)
+				.BorderImage(FAppStyle::Get().GetBrush("Brushes.Recessed"))
+				[
+					SNew(SVerticalBox)
+
+					+ SVerticalBox::Slot()
+					.AutoHeight()
+					.Padding(0)
+					[
+						SNew(SBorder)
+						.BorderImage(FAppStyle::Get().GetBrush("Brushes.Panel"))
+						[
+							SNew(SHorizontalBox)
+
+							+ SHorizontalBox::Slot()
+							.VAlign(VAlign_Center)
+							.FillWidth(1.0f)
+							.Padding(14, 0, 0, 0)
+							[
+								SNew(STextBlock)
+								.Text(LOCTEXT("ProjectLauncherCustomProfilesTitle", "Custom Launch Profiles"))
+								.Font(FAppStyle::Get().GetFontStyle("NormalFontBold"))
+							]
+
+							
+
+							+ SHorizontalBox::Slot()
+							.HAlign(HAlign_Right)
+							.AutoWidth()
+							.Padding(0, 0, 16, 0)
+							[
+								SNew(SPositiveActionButton)
+								.Text(LOCTEXT("AddButtonLabel", "Add"))
+								.ToolTipText(LOCTEXT("AddFilterToolTip", "Add a new custom launch profile using wizard"))
+								.OnGetMenuContent(this, &SLegacyProjectLauncher::MakeProfileWizardsMenu)
+							]
+						]
+					]
+					+ SVerticalBox::Slot()
+					.FillHeight(1)
+					.Padding(2)
+					[
+						SAssignNew(ProfileList, SBorder)
+						.BorderImage(FAppStyle::Get().GetBrush("Brushes.Recessed"))
+						.Padding(0)
+						[
+							// Simple Launch List
+							SNew(SProjectLauncherProfileListView, InModel)
+							.OnProfileEdit(this, &SLegacyProjectLauncher::OnProfileEdit)
+							.OnProfileRun(this, &SLegacyProjectLauncher::OnProfileRun)
+							.OnProfileDelete(this, &SLegacyProjectLauncher::OnProfileDelete)
+						]
+					]
+				]
+			]
+		]
+		
+		// Launch Settings
+		+ SWidgetSwitcher::Slot()
+		[
+			SAssignNew(ProfileSettingsPanel, SProjectLauncherSettings, InModel)
+			.OnCloseClicked(this, &SLegacyProjectLauncher::OnProfileSettingsClose)
+			.OnDeleteClicked(this, &SLegacyProjectLauncher::OnProfileDelete)
+		]
+
+		// Progress Panel
+		+ SWidgetSwitcher::Slot()
+		[
+			SAssignNew(ProgressPanel, SProjectLauncherProgress)
+			.OnCloseClicked(this, &SLegacyProjectLauncher::OnProgressClose)
+			.OnRerunClicked(this, &SLegacyProjectLauncher::OnRerunClicked)
+		]
+	];
+}
+END_SLATE_FUNCTION_BUILD_OPTIMIZATION
+
+
+/* SLegacyProjectLauncher implementation
+*****************************************************************************/
+
+void SLegacyProjectLauncher::FillWindowMenu(FMenuBuilder& MenuBuilder, TSharedRef<FWorkspaceItem> RootMenuGroup)
+{
+#if !WITH_EDITOR
+	MenuBuilder.BeginSection("WindowGlobalTabSpawners", LOCTEXT("UfeMenuGroup", "Unreal Frontend"));
+	{
+		FGlobalTabmanager::Get()->PopulateTabSpawnerMenu(MenuBuilder, RootMenuGroup);
+	}
+	MenuBuilder.EndSection();
+#endif //!WITH_EDITOR
+}
+
+
+/* SLegacyProjectLauncher callbacks
+*****************************************************************************/
+
+void SLegacyProjectLauncher::OnAdvancedChanged(const ECheckBoxState NewCheckedState)
+{
+	bAdvanced = (NewCheckedState == ECheckBoxState::Checked);
+}
+
+
+ECheckBoxState SLegacyProjectLauncher::OnIsAdvanced() const
+{
+	return (bAdvanced) ? ECheckBoxState::Checked : ECheckBoxState::Unchecked;
+}
+
+
+const FSlateBrush* SLegacyProjectLauncher::GetAdvancedToggleBrush() const
+{
+	return FAppStyle::Get().GetBrush("Icons.Advanced");
+}
+
+
+bool SLegacyProjectLauncher::GetIsAdvanced() const
+{
+	return bAdvanced;
+}
+
+
+void SLegacyProjectLauncher::OnProfileEdit(const ILauncherProfileRef& Profile)
+{
+	Model->SelectProfile(Profile);
+	WidgetSwitcher->SetActiveWidgetIndex(ELauncherPanels::ProfileEditor);
+}
+
+
+void SLegacyProjectLauncher::OnProfileRun(const ILauncherProfileRef& Profile)
+{
+	LauncherProfile = Profile;
+	LauncherWorker = Model->GetSProjectLauncher()->Launch(Model->GetDeviceProxyManager(), Profile);
+	
+	if (LauncherWorker.IsValid())
+	{
+		ProgressPanel->SetLauncherWorker(LauncherWorker.ToSharedRef());
+		WidgetSwitcher->SetActiveWidgetIndex(ELauncherPanels::Progress);
+	}
+}
+
+
+void SLegacyProjectLauncher::OnProfileDelete(const ILauncherProfileRef& Profile)
+{
+	Model->GetProfileManager()->RemoveProfile(Profile);
+}
+
+
+void SLegacyProjectLauncher::OnAddCustomLaunchProfileClicked()
+{
+	ILauncherProfileRef Profile = Model->GetProfileManager()->AddNewProfile();
+	
+	OnProfileEdit(Profile);
+
+	ProfileSettingsPanel->EnterEditMode();
+}
+
+
+TSharedRef<SWidget> SLegacyProjectLauncher::MakeProfileWizardsMenu()
+{
+	FMenuBuilder MenuBuilder(true, NULL);
+
+	MenuBuilder.BeginSection("Create", LOCTEXT("CreateSection", "CREATE"));
+
+	MenuBuilder.AddMenuEntry(
+		LOCTEXT("CustomProfileLabel", "Create Custom Profile"),
+		LOCTEXT("CustomProfileDescription", "Add a new custom launch profile."),
+		FSlateIcon(),
+		FUIAction(
+			FExecuteAction::CreateSP(this, &SLegacyProjectLauncher::OnAddCustomLaunchProfileClicked),
+			FCanExecuteAction()
+		)
+	);
+
+	MenuBuilder.EndSection();
+
+	MenuBuilder.BeginSection("CreateFromPreset", LOCTEXT("CreateFromPreset", "CREATE FROM PRESET"));
+
+	const TArray<ILauncherProfileWizardPtr>& Wizards = Model->GetProfileManager()->GetProfileWizards();
+	for (const ILauncherProfileWizardPtr& Wizard : Wizards)
+	{
+		FText WizardName = Wizard->GetName();
+		FText WizardDescription = Wizard->GetDescription();
+				
+		MenuBuilder.AddMenuEntry(
+			WizardName,
+			WizardDescription,
+			FSlateIcon(),
+			FUIAction(
+				FExecuteAction::CreateSP(this, &SLegacyProjectLauncher::ExecProfileWizard, Wizard),
+				FCanExecuteAction()
+				)
+			);
+	}
+
+	MenuBuilder.EndSection();
+
+	return MenuBuilder.MakeWidget();
+}
+
+
+void SLegacyProjectLauncher::ExecProfileWizard(ILauncherProfileWizardPtr InWizard)
+{
+	InWizard->HandleCreateLauncherProfile(Model->GetProfileManager());
+}
+
+
+FReply SLegacyProjectLauncher::OnProfileSettingsClose()
+{
+	if (LauncherWorker.IsValid())
+	{
+		LauncherWorker->Cancel();
+	}
+	LauncherProfile.Reset();
+
+	WidgetSwitcher->SetActiveWidgetIndex(ELauncherPanels::Launch);
+
+	return FReply::Handled();
+}
+
+
+FReply SLegacyProjectLauncher::OnProgressClose()
+{
+	if (LauncherWorker.IsValid())
+	{
+		LauncherWorker->Cancel();
+	}
+	LauncherProfile.Reset();
+
+	WidgetSwitcher->SetActiveWidgetIndex(ELauncherPanels::Launch);
+
+	return FReply::Handled();
+}
+
+
+FReply SLegacyProjectLauncher::OnRerunClicked()
+{
+	if (LauncherWorker.IsValid())
+	{
+		LauncherWorker->Cancel();
+	}
+	LauncherWorker = Model->GetSProjectLauncher()->Launch(Model->GetDeviceProxyManager(), LauncherProfile.ToSharedRef());
+
+	if (LauncherWorker.IsValid())
+	{
+		ProgressPanel->SetLauncherWorker(LauncherWorker.ToSharedRef());
+	}
+
+	return FReply::Handled();
+}
+
+
+#undef LOCTEXT_NAMESPACE
diff --git a/Engine/Source/Developer/LegacyProjectLauncher/Private/Widgets/SLegacyProjectLauncher.h b/Engine/Source/Developer/LegacyProjectLauncher/Private/Widgets/SLegacyProjectLauncher.h
new file mode 100644
--- /dev/null
+++ b/Engine/Source/Developer/LegacyProjectLauncher/Private/Widgets/SLegacyProjectLauncher.h
@@ -0,0 +1,144 @@
+// Copyright Epic Games, Inc. All Rights Reserved.
+
+#pragma once
+
+#include "CoreTypes.h"
+#include "Framework/Docking/WorkspaceItem.h"
+#include "Input/Reply.h"
+#include "Layout/Visibility.h"
+#include "SlateFwd.h"
+#include "Templates/SharedPointer.h"
+#include "Widgets/DeclarativeSyntaxSupport.h"
+#include "Widgets/SCompoundWidget.h"
+
+#include "Models/LegacyProjectLauncherModel.h"
+
+class FMenuBuilder;
+class SBorder;
+class SProjectLauncherProgress;
+class SProjectLauncherSettings;
+class SProjectLauncherSimpleDeviceListView;
+class SWindow;
+
+struct FSlateBrush;
+
+enum class ECheckBoxState : uint8;
+
+
+/**
+ * Implements a Slate widget for the launcher user interface.
+ */
+class SLegacyProjectLauncher
+	: public SCompoundWidget
+{
+public:
+
+	SLATE_BEGIN_ARGS(SLegacyProjectLauncher) { }
+
+		/** Exposes a delegate to be invoked when the launcher has closed. */
+		SLATE_EVENT(FSimpleDelegate, OnClosed)
+
+	SLATE_END_ARGS()
+
+public:
+
+	/** Constructor. */
+	SLegacyProjectLauncher();
+
+	/** Destructor. */
+	~SLegacyProjectLauncher();
+
+public:
+
+	/**
+	 * Constructs the widget.
+	 *
+	 * @param InArgs The Slate argument list.
+	 * @param ConstructUnderMajorTab The major tab which will contain the session front-end.
+	 * @param ConstructUnderWindow The window in which this widget is being constructed.
+	 * @param InModel The view model to use.
+	 */
+	void Construct(const FArguments& InArgs, const TSharedRef<SDockTab>& ConstructUnderMajorTab, const TSharedPtr<SWindow>& ConstructUnderWindow, const TSharedRef<FProjectLauncherModel>& InModel);
+
+protected:
+
+	/**
+	 * Fills the Window menu with menu items.
+	 *
+	 * @param MenuBuilder The multi-box builder that should be filled with content for this pull-down menu.
+	 * @param RootMenuGroup The root menu group.
+	 */
+	static void FillWindowMenu(FMenuBuilder& MenuBuilder, TSharedRef<FWorkspaceItem> RootMenuGroup);
+
+private:
+
+	/** Callback for toggling the advanced mode. */
+	void OnAdvancedChanged(const ECheckBoxState NewCheckedState);
+
+	/** Callback to determine if we are in advanced mode. */
+	ECheckBoxState OnIsAdvanced() const;
+
+	/** Get advanced toggle button brush. */
+	const FSlateBrush* GetAdvancedToggleBrush() const;
+
+	/** Callback for whether advanced options should be shown. */
+	bool GetIsAdvanced() const;
+
+	/** Callback for editing a profile. */
+	void OnProfileEdit(const ILauncherProfileRef& Profile);
+
+	/** Callback for running a profile. */
+	void OnProfileRun(const ILauncherProfileRef& Profile);
+
+	/** Callback for deleting a profile. */
+	void OnProfileDelete(const ILauncherProfileRef& Profile);
+
+	/** Callback for clicking the Add New Launch Profile. */
+	void OnAddCustomLaunchProfileClicked();
+
+	/** Callback for filling profile wizard menu. */
+	TSharedRef<SWidget> MakeProfileWizardsMenu();
+
+	/** Execute specified profile wizard. */
+	void ExecProfileWizard(ILauncherProfileWizardPtr InWizard);
+
+	/** Callback for when the settings panel is closed. */
+	FReply OnProfileSettingsClose();
+
+	/** Callback for when the progress panel is closed. */
+	FReply OnProgressClose();
+
+	/** Callback for when the progress panel requests the profile to be re run. */
+	FReply OnRerunClicked();
+
+private:
+
+	/** The current launcher worker, if any. */
+	ILauncherWorkerPtr LauncherWorker;
+
+	/** The Launcher profile the launcherWorker is running. */
+	ILauncherProfilePtr LauncherProfile;
+
+	/** Holds a pointer to the view model. */
+	TSharedPtr<FProjectLauncherModel> Model;
+
+	/** The profile settings panel. */
+	TSharedPtr<SProjectLauncherSettings> ProfileSettingsPanel;
+
+	/** The progress panel. */
+	TSharedPtr<SProjectLauncherProgress> ProgressPanel;
+
+	/** The widget switcher. */
+	TSharedPtr<SWidgetSwitcher> WidgetSwitcher;
+
+	/** Contains the launch list widgets. */
+//	TSharedPtr<SBorder> LaunchList;
+
+	TSharedPtr<SProjectLauncherSimpleDeviceListView> LaunchList;
+
+	/** Contains the profile list widgets. */
+	TSharedPtr<SBorder> ProfileList;
+
+	/** Whether we are showing advanced options. */
+	bool bAdvanced;
+};
diff --git a/Engine/Source/Developer/LegacyProjectLauncher/Private/Widgets/SProjectLauncher.cpp b/Engine/Source/Developer/LegacyProjectLauncher/Private/Widgets/SProjectLauncher.cpp
deleted file mode 100644
diff --git a/Engine/Source/Developer/LegacyProjectLauncher/Private/Widgets/SProjectLauncher.h b/Engine/Source/Developer/LegacyProjectLauncher/Private/Widgets/SProjectLauncher.h
deleted file mode 100644
diff --git a/Engine/Source/Developer/LegacyProjectLauncher/Private/Widgets/Settings/SProjectLauncherSettings.cpp b/Engine/Source/Developer/LegacyProjectLauncher/Private/Widgets/Settings/SProjectLauncherSettings.cpp
--- a/Engine/Source/Developer/LegacyProjectLauncher/Private/Widgets/Settings/SProjectLauncherSettings.cpp
+++ b/Engine/Source/Developer/LegacyProjectLauncher/Private/Widgets/Settings/SProjectLauncherSettings.cpp
@@ -17,7 +17,7 @@
 #include "Widgets/Text/SInlineEditableTextBlock.h"
 
 #include "Models/ProjectLauncherCommands.h"
-#include "Models/ProjectLauncherModel.h"
+#include "Models/LegacyProjectLauncherModel.h"
 #include "Widgets/Archive/SProjectLauncherArchivePage.h"
 #include "Widgets/Build/SProjectLauncherBuildPage.h"
 #include "Widgets/Cook/SProjectLauncherCookPage.h"
diff --git a/Engine/Source/Developer/LegacyProjectLauncher/Private/Widgets/Shared/ProjectLauncherDelegates.h b/Engine/Source/Developer/LegacyProjectLauncher/Private/Widgets/Shared/ProjectLauncherDelegates.h
--- a/Engine/Source/Developer/LegacyProjectLauncher/Private/Widgets/Shared/ProjectLauncherDelegates.h
+++ b/Engine/Source/Developer/LegacyProjectLauncher/Private/Widgets/Shared/ProjectLauncherDelegates.h
@@ -7,8 +7,8 @@
 
 
 /**
- * Delegate type for SProjectLauncher Profile running.
+ * Delegate type for SLegacyProjectLauncher Profile running.
  *
- * The first parameter is the SProjectLauncher Profile that you want to run.
+ * The first parameter is the SLegacyProjectLauncher Profile that you want to run.
  */
 DECLARE_DELEGATE_OneParam(FOnProfileRun, const ILauncherProfileRef&)
diff --git a/Engine/Source/Developer/LegacyProjectLauncher/Private/Widgets/Shared/SProjectLauncherBuildTargetSelector.h b/Engine/Source/Developer/LegacyProjectLauncher/Private/Widgets/Shared/SProjectLauncherBuildTargetSelector.h
--- a/Engine/Source/Developer/LegacyProjectLauncher/Private/Widgets/Shared/SProjectLauncherBuildTargetSelector.h
+++ b/Engine/Source/Developer/LegacyProjectLauncher/Private/Widgets/Shared/SProjectLauncherBuildTargetSelector.h
@@ -9,7 +9,7 @@
 #include "Widgets/SWidget.h"
 #include "Widgets/SCompoundWidget.h"
 #include "Widgets/Input/SComboBox.h"
-#include "Models/ProjectLauncherModel.h"
+#include "Models/LegacyProjectLauncherModel.h"
 
 class Error;
 
diff --git a/Engine/Source/Developer/LegacyProjectLauncher/Private/Widgets/Shared/SProjectLauncherProfileNameDescEditor.h b/Engine/Source/Developer/LegacyProjectLauncher/Private/Widgets/Shared/SProjectLauncherProfileNameDescEditor.h
--- a/Engine/Source/Developer/LegacyProjectLauncher/Private/Widgets/Shared/SProjectLauncherProfileNameDescEditor.h
+++ b/Engine/Source/Developer/LegacyProjectLauncher/Private/Widgets/Shared/SProjectLauncherProfileNameDescEditor.h
@@ -7,7 +7,7 @@
 #include "Misc/Attribute.h"
 #include "Widgets/DeclarativeSyntaxSupport.h"
 #include "Widgets/SCompoundWidget.h"
-#include "Models/ProjectLauncherModel.h"
+#include "Models/LegacyProjectLauncherModel.h"
 #include "Styling/AppStyle.h"
 
 /**
