prevent legacy project launcher from being able to edit or launch profiles with multiple UAT commands

#rnx
#Synced-CL 48803560

diff --git a/Engine/Source/Developer/LauncherServices/Private/Profiles/LauncherProfile.h b/Engine/Source/Developer/LauncherServices/Private/Profiles/LauncherProfile.h
--- a/Engine/Source/Developer/LauncherServices/Private/Profiles/LauncherProfile.h
+++ b/Engine/Source/Developer/LauncherServices/Private/Profiles/LauncherProfile.h
@@ -3332,6 +3332,11 @@ PRAGMA_ENABLE_DEPRECATION_WARNINGS
 		Validate();
 	}
 
+	virtual bool SupportsLegacyProjectLauncher() const override
+	{
+		return (UATCommands.Num() == 0);
+	}
+
 
 	//~ End ILauncherProfile Interface
 
diff --git a/Engine/Source/Developer/LauncherServices/Public/ILauncherProfile.h b/Engine/Source/Developer/LauncherServices/Public/ILauncherProfile.h
--- a/Engine/Source/Developer/LauncherServices/Public/ILauncherProfile.h
+++ b/Engine/Source/Developer/LauncherServices/Public/ILauncherProfile.h
@@ -1778,6 +1778,9 @@ public:
 
 public:
 
+	/** Whether this profile can be used by the legacy Project Launcher */
+	virtual bool SupportsLegacyProjectLauncher() const = 0;
+
 	/** Virtual destructor. */
 	virtual ~ILauncherProfile( ) = default;
 };
diff --git a/Engine/Source/Developer/LegacyProjectLauncher/Private/Widgets/Profile/SProjectLauncherProfileListRow.h b/Engine/Source/Developer/LegacyProjectLauncher/Private/Widgets/Profile/SProjectLauncherProfileListRow.h
--- a/Engine/Source/Developer/LegacyProjectLauncher/Private/Widgets/Profile/SProjectLauncherProfileListRow.h
+++ b/Engine/Source/Developer/LegacyProjectLauncher/Private/Widgets/Profile/SProjectLauncherProfileListRow.h
@@ -102,6 +102,7 @@ public:
 						.OnClicked(this, &SProjectLauncherProfileListRow::OnEditClicked)
 						.ToolTipText(LOCTEXT("EditProfileToolTipText", "Edit profile."))
 						.Icon(this, &SProjectLauncherProfileListRow::GetEditIcon)
+						.IsEnabled_Lambda( [this]() { return LaunchProfile.IsValid() && LaunchProfile->SupportsLegacyProjectLauncher(); } )
 					]
 
 					+ SHorizontalBox::Slot()
diff --git a/Engine/Source/Developer/LegacyProjectLauncher/Private/Widgets/Shared/SProjectLauncherProfileLaunchButton.h b/Engine/Source/Developer/LegacyProjectLauncher/Private/Widgets/Shared/SProjectLauncherProfileLaunchButton.h
--- a/Engine/Source/Developer/LegacyProjectLauncher/Private/Widgets/Shared/SProjectLauncherProfileLaunchButton.h
+++ b/Engine/Source/Developer/LegacyProjectLauncher/Private/Widgets/Shared/SProjectLauncherProfileLaunchButton.h
@@ -87,6 +87,14 @@ public:
 
 			return ErrorToolTipWidget;
 		}
+		else if (!IsLegacyProfile())
+		{
+			ErrorToolTipWidget->SetContentWidget(
+				SNew(STextBlock)
+				.Text(LOCTEXT("ProjectLauncherNotLegacyProfile", "This profile can only be used with the new Project Launcher")));
+
+			return ErrorToolTipWidget;
+		}
 		return SWidget::GetToolTip();
 	}
 
@@ -103,10 +111,16 @@ private:
 		return true;
 	}
 
+	bool IsLegacyProfile() const
+	{
+		ILauncherProfilePtr LaunchProfile = LaunchProfileAttr.Get();
+		return LaunchProfile.IsValid() && LaunchProfile->SupportsLegacyProjectLauncher();
+	};
+
 	// Callback to see if the launch button should be enabled
 	bool ButtonEnabled() const
 	{
-		return !HasError();
+		return !HasError() && IsLegacyProfile();
 	}
 
 	// Get the SlateIcon for Launch Button
