[ProjectLauncher]
- Fix crash caused by linking of similarly named linker entry between the legacy and the new project launcher that can happen when compiling in monolithic (i.e. QAGCookedEditor)

#Synced-CL 49408961

diff --git a/Engine/Source/Developer/LegacyProjectLauncher/Private/ProjectLauncherModule.cpp b/Engine/Source/Developer/LegacyProjectLauncher/Private/ProjectLauncherModule.cpp
--- a/Engine/Source/Developer/LegacyProjectLauncher/Private/ProjectLauncherModule.cpp
+++ b/Engine/Source/Developer/LegacyProjectLauncher/Private/ProjectLauncherModule.cpp
@@ -23,8 +23,8 @@ static const FName ProjectLauncherTabName("LegacyProjectLauncher");
 /**
  * Implements the SessionSProjectLauncher module.
  */
-class FProjectLauncherModule
-	: public IProjectLauncherModule
+class FLegacyProjectLauncherModule
+	: public ILegacyProjectLauncherModule
 {
 public:
 
@@ -32,7 +32,7 @@ public:
 	
 	virtual void StartupModule() override
 	{
-		FGlobalTabmanager::Get()->RegisterNomadTabSpawner(ProjectLauncherTabName, FOnSpawnTab::CreateRaw(this, &FProjectLauncherModule::SpawnProjectLauncherTab))
+		FGlobalTabmanager::Get()->RegisterNomadTabSpawner(ProjectLauncherTabName, FOnSpawnTab::CreateRaw(this, &FLegacyProjectLauncherModule::SpawnProjectLauncherTab))
 			.SetDisplayName(NSLOCTEXT("FProjectLauncherModule", "ProjectLauncherTabTitle", "Legacy Project Launcher"))
 			.SetTooltipText(NSLOCTEXT("FProjectLauncherModule", "ProjectLauncherTooltipText", "Open the Legacy Project Launcher tab."))
 			.SetIcon(FSlateIcon(FAppStyle::GetAppStyleSetName(), "Launcher.TabIcon"))
@@ -73,4 +73,4 @@ private:
 };
 
 
-IMPLEMENT_MODULE(FProjectLauncherModule, LegacyProjectLauncher);
+IMPLEMENT_MODULE(FLegacyProjectLauncherModule, LegacyProjectLauncher);
diff --git a/Engine/Source/Developer/LegacyProjectLauncher/Public/IProjectLauncherModule.h b/Engine/Source/Developer/LegacyProjectLauncher/Public/IProjectLauncherModule.h
--- a/Engine/Source/Developer/LegacyProjectLauncher/Public/IProjectLauncherModule.h
+++ b/Engine/Source/Developer/LegacyProjectLauncher/Public/IProjectLauncherModule.h
@@ -8,11 +8,11 @@
 /**
  * Interface for launcher UI modules.
  */
-class IProjectLauncherModule
+class ILegacyProjectLauncherModule
 	: public IModuleInterface
 {
 public:
 
 	/** Virtual destructor. */
-	virtual ~IProjectLauncherModule() = default;
+	virtual ~ILegacyProjectLauncherModule() = default;
 };
