fix for legacy project launcher error widget + add static assert to catch it earlier next time

#Synced-CL 50244055

diff --git a/Engine/Source/Developer/LegacyProjectLauncher/Private/Widgets/Shared/SProjectLauncherValidation.h b/Engine/Source/Developer/LegacyProjectLauncher/Private/Widgets/Shared/SProjectLauncherValidation.h
--- a/Engine/Source/Developer/LegacyProjectLauncher/Private/Widgets/Shared/SProjectLauncherValidation.h
+++ b/Engine/Source/Developer/LegacyProjectLauncher/Private/Widgets/Shared/SProjectLauncherValidation.h
@@ -201,6 +201,12 @@ public:
 				MakeValidationMessage(TEXT("Icons.Error"), ELauncherProfileValidationErrors::ZenPakStreamingRequiresDeployAndLaunch)
 			];
 
+		VertBox->AddSlot().AutoHeight()
+			[
+				MakeValidationMessage(TEXT("Icons.Error"), ELauncherProfileValidationErrors::AutomatedTestRequiredBuildCookRun)
+			];
+
+		static_assert(ELauncherProfileValidationErrors::Count == 32, "SProjectLauncherValidation::Construct() needs to be updated to account for modified enum values.");
 		check(VertBox->NumSlots() == ELauncherProfileValidationErrors::Count);
 
 		ChildSlot[VertBox.ToSharedRef()];
